Add missing include.

diff --git a/Engine/Plugins/Experimental/ComputeFramework/Source/ComputeFramework/Private/ComputeKernelShaderMap.cpp b/Engine/Plugins/Experimental/ComputeFramework/Source/ComputeFramework/Private/ComputeKernelShaderMap.cpp
--- a/Engine/Plugins/Experimental/ComputeFramework/Source/ComputeFramework/Private/ComputeKernelShaderMap.cpp
+++ b/Engine/Plugins/Experimental/ComputeFramework/Source/ComputeFramework/Private/ComputeKernelShaderMap.cpp
@@ -6,6 +6,7 @@
 #include "ComputeFramework/ComputeKernelDerivedDataVersion.h"
 #include "ComputeFramework/ComputeKernelShaderCompilationManager.h"
 #include "DerivedDataCacheInterface.h"
+#include "Interfaces/ITargetPlatformManagerModule.h"
 #include "ProfilingDebugging/CookStats.h"
 #include "ProfilingDebugging/DiagnosticTable.h"
 #include "Serialization/MemoryWriter.h"
