No unity fixes
#rnx

diff --git a/Engine/Plugins/Experimental/ComputeFramework/Source/ComputeFrameworkEditor/Private/ComputeFramework/ComputeFrameworkEditorModule.h b/Engine/Plugins/Experimental/ComputeFramework/Source/ComputeFrameworkEditor/Private/ComputeFramework/ComputeFrameworkEditorModule.h
--- a/Engine/Plugins/Experimental/ComputeFramework/Source/ComputeFrameworkEditor/Private/ComputeFramework/ComputeFrameworkEditorModule.h
+++ b/Engine/Plugins/Experimental/ComputeFramework/Source/ComputeFrameworkEditor/Private/ComputeFramework/ComputeFrameworkEditorModule.h
@@ -3,6 +3,7 @@
 #pragma once
 
 #include "ComputeFramework/IComputeFrameworkEditorModule.h"
+#include "ComputeFramework/ComputeFrameworkCompilationTick.h"
 
 class FComputeFrameworkEditorModule : public IComputeFrameworkEditorModule
 {
diff --git a/Engine/Source/Editor/UnrealEd/Private/WorldPartition/WorldPartitionResaveActorsBuilder.cpp b/Engine/Source/Editor/UnrealEd/Private/WorldPartition/WorldPartitionResaveActorsBuilder.cpp
--- a/Engine/Source/Editor/UnrealEd/Private/WorldPartition/WorldPartitionResaveActorsBuilder.cpp
+++ b/Engine/Source/Editor/UnrealEd/Private/WorldPartition/WorldPartitionResaveActorsBuilder.cpp
@@ -8,6 +8,7 @@
 #include "Engine/World.h"
 #include "GameFramework/Actor.h"
 #include "Commandlets/Commandlet.h"
+#include "Kismet2/BlueprintEditorUtils.h"
 #include "UObject/MetaData.h"
 #include "UObject/SavePackage.h"
 #include "Kismet2/BlueprintEditorUtils.h"
