Optimus: Add support for more shader parameter types.

#rnx

diff --git a/Engine/Plugins/Experimental/ComputeFramework/Source/ComputeFramework/Private/ComputeKernelSource.cpp b/Engine/Plugins/Experimental/ComputeFramework/Source/ComputeFramework/Private/ComputeKernelSource.cpp
--- a/Engine/Plugins/Experimental/ComputeFramework/Source/ComputeFramework/Private/ComputeKernelSource.cpp
+++ b/Engine/Plugins/Experimental/ComputeFramework/Source/ComputeFramework/Private/ComputeKernelSource.cpp
@@ -4,29 +4,61 @@
 
 #include "ShaderParameterMetadataBuilder.h"
 
-void UComputeKernelSource::GetShaderParameters(class FShaderParametersMetadataBuilder& OutBuilder) const
+template<typename T>
+void ParametrizedAddParm(FShaderParametersMetadataBuilder& OutBuilder, const TCHAR* InName)
+{
+	OutBuilder.AddParam<T>(InName);
+}
+
+
+static bool AddParamForType(
+	FShaderParametersMetadataBuilder& OutBuilder,
+	const FShaderParamTypeDefinition& InDef
+	)
+{
+	// Add param functions for various shader argument types. These match the ones found
+	// as specializations of TShaderParameterTypeInfo.
+	using AddParamFuncType = TFunction<void(FShaderParametersMetadataBuilder&, const TCHAR*)>; 
+	static TMap<FShaderValueType, AddParamFuncType> AddParamFuncs = {
+		// bool
+		{*FShaderValueType::Get(EShaderFundamentalType::Bool), &ParametrizedAddParm<bool>},
+		
+		// int
+		{*FShaderValueType::Get(EShaderFundamentalType::Int), &ParametrizedAddParm<int32>},
+		{*FShaderValueType::Get(EShaderFundamentalType::Int, 2), &ParametrizedAddParm<FIntPoint>},
+		{*FShaderValueType::Get(EShaderFundamentalType::Int, 3), &ParametrizedAddParm<FIntVector>},
+		{*FShaderValueType::Get(EShaderFundamentalType::Int, 4), &ParametrizedAddParm<FIntVector4>},
+
+		// uint
+		{*FShaderValueType::Get(EShaderFundamentalType::Uint), &ParametrizedAddParm<uint32>},
+		{*FShaderValueType::Get(EShaderFundamentalType::Uint, 2), &ParametrizedAddParm<FUintVector2>},
+		// -- NOTE: No FUintVector defined for TShaderParameterTypeInfo
+		{*FShaderValueType::Get(EShaderFundamentalType::Uint, 4), &ParametrizedAddParm<FUintVector4>},
+
+		// float
+		{*FShaderValueType::Get(EShaderFundamentalType::Float), &ParametrizedAddParm<float>},
+		{*FShaderValueType::Get(EShaderFundamentalType::Float, 2), &ParametrizedAddParm<FVector2f>},
+		{*FShaderValueType::Get(EShaderFundamentalType::Float, 3), &ParametrizedAddParm<FVector3f>},
+		{*FShaderValueType::Get(EShaderFundamentalType::Float, 4), &ParametrizedAddParm<FVector4f>},
+		{*FShaderValueType::Get(EShaderFundamentalType::Float, 4, 4), &ParametrizedAddParm<FMatrix44f>},
+	};
+
+	if (const AddParamFuncType* Entry = AddParamFuncs.Find(*InDef.ValueType))
+	{
+		(*Entry)(OutBuilder, *InDef.Name);
+		return true;
+	}
+	
+	return false;
+}
+
+
+
+void UComputeKernelSource::GetShaderParameters(FShaderParametersMetadataBuilder& OutBuilder) const
 {
 	for (auto& Input : InputParams)
 	{
-		ensureAlways(Input.ValueType->DimensionType == EShaderFundamentalDimensionType::Scalar);
-
-		switch (Input.ValueType->Type)
-		{
-		case EShaderFundamentalType::Bool:
-			OutBuilder.AddParam<bool>(*Input.Name);
-			break;
-
-		case EShaderFundamentalType::Int:
-			OutBuilder.AddParam<int32>(*Input.Name);
-			break;
-
-		case EShaderFundamentalType::Uint:
-			OutBuilder.AddParam<uint32>(*Input.Name);
-			break;
-
-		case EShaderFundamentalType::Float:
-			OutBuilder.AddParam<float>(*Input.Name);
-			break;
-		}
+		verifyf(AddParamForType(OutBuilder, Input),
+			TEXT("Unsupported shader param type (%s) for param '%s'"), *Input.ValueType->ToString(), *Input.Name);
 	}
 }
