ComputeFramework: Fix ref count race condition.
#rnx

diff --git a/Engine/Plugins/Experimental/ComputeFramework/Source/ComputeFramework/Private/ComputeKernelShaderMap.cpp b/Engine/Plugins/Experimental/ComputeFramework/Source/ComputeFramework/Private/ComputeKernelShaderMap.cpp
--- a/Engine/Plugins/Experimental/ComputeFramework/Source/ComputeFramework/Private/ComputeKernelShaderMap.cpp
+++ b/Engine/Plugins/Experimental/ComputeFramework/Source/ComputeFramework/Private/ComputeKernelShaderMap.cpp
@@ -646,13 +646,13 @@ void FComputeKernelShaderMap::Register(EShaderPlatform InShaderPlatform)
 void FComputeKernelShaderMap::AddRef()
 {
 	check(!bDeletedThroughDeferredCleanup);
-	++NumRefs;
+	FPlatformAtomics::InterlockedIncrement(&NumRefs);
 }
 
 void FComputeKernelShaderMap::Release()
 {
 	check(NumRefs > 0);
-	if(--NumRefs == 0)
+	if(FPlatformAtomics::InterlockedDecrement(&NumRefs) == 0)
 	{
 		if (bRegistered)
 		{
