ComputeFramework: Associate shadermaps with assets to enable shader loading priority support.
#rnx

diff --git a/Engine/Plugins/Experimental/ComputeFramework/Source/ComputeFramework/Private/ComputeFramework/ComputeKernelShared.h b/Engine/Plugins/Experimental/ComputeFramework/Source/ComputeFramework/Private/ComputeFramework/ComputeKernelShared.h
--- a/Engine/Plugins/Experimental/ComputeFramework/Source/ComputeFramework/Private/ComputeFramework/ComputeKernelShared.h
+++ b/Engine/Plugins/Experimental/ComputeFramework/Source/ComputeFramework/Private/ComputeFramework/ComputeKernelShared.h
@@ -426,7 +426,8 @@ public:
 		uint64 InShaderCodeHash,
 		FComputeKernelDefinitionSet& InShaderDefinitionSet,
 		FComputeKernelPermutationVector& InShaderPermutationVector,
-		FShaderParametersMetadata* InShaderMetadata
+		FShaderParametersMetadata* InShaderMetadata,
+		FName const& InAssetPath
 	);
 
 	void SetupCompileEnvironment(int32 InPermutationId, FShaderCompilerEnvironment& OutShaderEnvironment) const;
@@ -536,6 +537,11 @@ private:
 
 	FString FriendlyName;
 
+#if WITH_EDITOR
+	/** Asset using this resource */
+	FName AssetPath;
+#endif
+
 	friend class FComputeKernelShaderMap;
 	friend class FShaderCompilingManager;
 };
diff --git a/Engine/Plugins/Experimental/ComputeFramework/Source/ComputeFramework/Private/ComputeGraph.cpp b/Engine/Plugins/Experimental/ComputeFramework/Source/ComputeFramework/Private/ComputeGraph.cpp
--- a/Engine/Plugins/Experimental/ComputeFramework/Source/ComputeFramework/Private/ComputeGraph.cpp
+++ b/Engine/Plugins/Experimental/ComputeFramework/Source/ComputeFramework/Private/ComputeGraph.cpp
@@ -486,7 +486,8 @@ void UComputeGraph::CacheResourceShadersForRendering(uint32 CompilationFlags)
 				ShaderSourceHash, 
 				ShaderDefinitionSet,
 				ShaderPermutationVector,
-				ShaderMetadata);
+				ShaderMetadata,
+				GetOutermost()->GetFName());
 
 			KernelResource->OnCompilationComplete().BindUObject(this, &UComputeGraph::ShaderCompileCompletionCallback);
 
@@ -600,7 +601,8 @@ void UComputeGraph::BeginCacheForCookedPlatformData(ITargetPlatform const* Targe
 					ShaderSourceHash, 
 					ShaderDefinitionSet,
 					ShaderPermutationVector,
-					ShaderMetadata);
+					ShaderMetadata,
+					GetOutermost()->GetFName());
 
 				const uint32 CompilationFlags = uint32(EComputeKernelCompilationFlags::IsCooking);
 				CacheShadersForResource(ShaderPlatform, TargetPlatform, CompilationFlags, KernelResource.Get());
diff --git a/Engine/Plugins/Experimental/ComputeFramework/Source/ComputeFramework/Private/ComputeKernelShared.cpp b/Engine/Plugins/Experimental/ComputeFramework/Source/ComputeFramework/Private/ComputeKernelShared.cpp
--- a/Engine/Plugins/Experimental/ComputeFramework/Source/ComputeFramework/Private/ComputeKernelShared.cpp
+++ b/Engine/Plugins/Experimental/ComputeFramework/Source/ComputeFramework/Private/ComputeKernelShared.cpp
@@ -172,6 +172,8 @@ void FComputeKernelResource::SerializeShaderMap(FArchive& Ar)
 
 			if (bValid)
 			{
+				GameThreadShaderMap->AssociateWithAsset(AssetPath);
+
 				GameThreadShaderMap->Serialize(Ar);
 			}
 #endif
@@ -205,7 +207,8 @@ void FComputeKernelResource::SetupResource(
 	uint64 InShaderCodeHash,
 	FComputeKernelDefinitionSet& InShaderDefinitionSet,
 	FComputeKernelPermutationVector& InShaderPermutationVector,
-	FShaderParametersMetadata* InShaderMetadata
+	FShaderParametersMetadata* InShaderMetadata,
+	FName const& InAssetPath
 	)
 {
 	FeatureLevel = InFeatureLevel;
@@ -217,6 +220,9 @@ void FComputeKernelResource::SetupResource(
 	ShaderPermutationVector = MoveTemp(InShaderPermutationVector);
 	ShaderMetadata.Reset(InShaderMetadata);
 	CompileErrors.Reset();
+#if WITH_EDITOR
+	AssetPath = InAssetPath;
+#endif
 }
 
 int32 FComputeKernelResource::GetNumPermutations() const 
