ComputeFramework: Remove deprecated KERNEL_PARAM define.
#rnx

diff --git a/Engine/Plugins/Experimental/ComputeFramework/Shaders/Private/ComputeKernelCommon.ush b/Engine/Plugins/Experimental/ComputeFramework/Shaders/Private/ComputeKernelCommon.ush
--- a/Engine/Plugins/Experimental/ComputeFramework/Shaders/Private/ComputeKernelCommon.ush
+++ b/Engine/Plugins/Experimental/ComputeFramework/Shaders/Private/ComputeKernelCommon.ush
@@ -12,10 +12,6 @@
         uint3 GTid: SV_GroupThreadID, \
         uint3 DTid : SV_DispatchThreadID)
 
-// Macros to directly expose parameters for a kernel.
-
-#define KERNEL_PARAM(Type, Name)    Type Name;
-
 // Macros to expose external functions for a kernel.
 // A kernel assumes these functions exist and are provided by its data providers.
 
