ComputeFramework: Fix non-deterministic DDC key for shader. Was causing slow start up time due to shader recompilation.

diff --git a/Engine/Plugins/Experimental/ComputeFramework/Source/ComputeFramework/Private/ComputeGraph.cpp b/Engine/Plugins/Experimental/ComputeFramework/Source/ComputeFramework/Private/ComputeGraph.cpp
--- a/Engine/Plugins/Experimental/ComputeFramework/Source/ComputeFramework/Private/ComputeGraph.cpp
+++ b/Engine/Plugins/Experimental/ComputeFramework/Source/ComputeFramework/Private/ComputeGraph.cpp
@@ -504,7 +504,7 @@ void UComputeGraph::CacheResourceShadersForRendering(uint32 CompilationFlags)
 
 			FString ShaderEntryPoint = Kernel->KernelSource->GetEntryPoint();
 			FString ShaderSource = BuildKernelSource(KernelIndex, ShaderDefinitionSet, ShaderPermutationVector);
-			uint64 ShaderSourceHash = FCrc::TypeCrc32(*ShaderSource, Kernel->KernelSource->GetSourceCodeHash());
+			const uint64 ShaderSourceHash = HashCombine(GetTypeHash(*ShaderSource), Kernel->KernelSource->GetSourceCodeHash());
 			FShaderParametersMetadata* ShaderMetadata = BuildKernelShaderMetadata(KernelIndex);
 
 			const ERHIFeatureLevel::Type CacheFeatureLevel = GMaxRHIFeatureLevel;
diff --git a/Engine/Plugins/Experimental/ComputeFramework/Source/ComputeFramework/Private/ComputeKernelFromText.cpp b/Engine/Plugins/Experimental/ComputeFramework/Source/ComputeFramework/Private/ComputeKernelFromText.cpp
--- a/Engine/Plugins/Experimental/ComputeFramework/Source/ComputeFramework/Private/ComputeKernelFromText.cpp
+++ b/Engine/Plugins/Experimental/ComputeFramework/Source/ComputeFramework/Private/ComputeKernelFromText.cpp
@@ -199,15 +199,14 @@ void UComputeKernelFromText::ReparseKernelSourceText()
 	}
 
 	uint64 NewHash = 0;
-
-	NewHash = FCrc::TypeCrc32(UniqueId, NewHash);
-	NewHash = FCrc::TypeCrc32(*SourceFile.FilePath, NewHash);
-	NewHash = FCrc::TypeCrc32(*KernelSourceText, NewHash);
+	NewHash = HashCombine(GetTypeHash(UniqueId), NewHash);
+	NewHash = HashCombine(GetTypeHash(SourceFile.FilePath), NewHash);
+	NewHash = HashCombine(GetTypeHash(KernelSourceText), NewHash);
 
 	if (SourceHash != NewHash)
 	{
 		SourceHash = NewHash;
-		// todo[CF]: Nofify graphs for recompilation
+		// todo[CF]: Notify graphs for recompilation
 	}
 }
 
diff --git a/Engine/Plugins/Experimental/ComputeFramework/Source/ComputeFramework/Private/ComputeKernelShaderMap.cpp b/Engine/Plugins/Experimental/ComputeFramework/Source/ComputeFramework/Private/ComputeKernelShaderMap.cpp
--- a/Engine/Plugins/Experimental/ComputeFramework/Source/ComputeFramework/Private/ComputeKernelShaderMap.cpp
+++ b/Engine/Plugins/Experimental/ComputeFramework/Source/ComputeFramework/Private/ComputeKernelShaderMap.cpp
@@ -254,22 +254,12 @@ FComputeKernelShaderMap* FComputeKernelShaderMap::FindId(const FComputeKernelSha
 	return Result;
 }
 
-void ComputeKernelShaderMapAppendKeyString(EShaderPlatform InPlatform, FString& OutKeyString)
-{
-#if WITH_EDITOR
-	//Keep library version in the DDC key to invalidate it once we move to a new library
-	OutKeyString += TEXT("ComputeKernelVersion");
-	OutKeyString += TEXT("_");
-#endif //WITH_EDITOR
-}
-
 /** Creates a string key for the derived data cache given a shader map id. */
 static FString GetComputeKernelShaderMapKeyString(const FComputeKernelShaderMapId& InShaderMapId, EShaderPlatform InPlatform)
 {
 #if WITH_EDITOR
 	const FName Format = LegacyShaderPlatformToShaderFormat(InPlatform);
 	FString ShaderMapKeyString = Format.ToString() + TEXT("_") + FString(FString::FromInt(GetTargetPlatformManagerRef().ShaderFormatVersion(Format))) + TEXT("_");
-	ComputeKernelShaderMapAppendKeyString(InPlatform, ShaderMapKeyString);
 	ShaderMapAppendKeyString(InPlatform, ShaderMapKeyString);
 	InShaderMapId.AppendKeyString(ShaderMapKeyString);
 	return FDerivedDataCacheInterface::BuildCacheKey(TEXT("CKERSM"), COMPUTEKERNEL_DERIVEDDATA_VER, *ShaderMapKeyString);
@@ -329,12 +319,16 @@ void FComputeKernelShaderMap::SaveToDerivedDataCache()
 {
 #if WITH_EDITOR
 	COOK_STAT(auto Timer = ComputeKernelShaderCookStats::UsageStats.TimeSyncWork());
-	TArray<uint8> SaveData;
+	
+TArray<uint8> SaveData;
 	FMemoryWriter Ar(SaveData, true);
 	Serialize(Ar);
 
-	GetDerivedDataCacheRef().Put(*GetComputeKernelShaderMapKeyString(GetContent()->ShaderMapId, GetShaderPlatform()), SaveData, FStringView(*GetFriendlyName()));
-	COOK_STAT(Timer.AddMiss(SaveData.Num()));
+	const FString DataKey = GetComputeKernelShaderMapKeyString(GetContent()->ShaderMapId, GetShaderPlatform());
+
+	GetDerivedDataCacheRef().Put(*DataKey, SaveData, GetFriendlyName());
+	
+COOK_STAT(Timer.AddMiss(SaveData.Num()));
 #endif
 }
 
