ComputeFramework: Add support for data provider contribution to shader hash. Will now trigger recompilation if data provider source or dependencies change.

diff --git a/Engine/Plugins/Experimental/ComputeFramework/Source/ComputeFramework/Private/ComputeFramework/ComputeKernelShared.h b/Engine/Plugins/Experimental/ComputeFramework/Source/ComputeFramework/Private/ComputeFramework/ComputeKernelShared.h
--- a/Engine/Plugins/Experimental/ComputeFramework/Source/ComputeFramework/Private/ComputeFramework/ComputeKernelShared.h
+++ b/Engine/Plugins/Experimental/ComputeFramework/Source/ComputeFramework/Private/ComputeFramework/ComputeKernelShared.h
@@ -422,8 +422,8 @@ public:
 		ERHIFeatureLevel::Type InFeatureLevel,
 		FString const& InFriendlyName,
 		FString const& InShaderEntryPoint,
+		FString const& InShaderHashKey,
 		FString InShaderSource,
-		uint64 InShaderCodeHash,
 		FComputeKernelDefinitionSet& InShaderDefinitionSet,
 		FComputeKernelPermutationVector& InShaderPermutationVector,
 		FShaderParametersMetadata* InShaderMetadata,
diff --git a/Engine/Plugins/Experimental/ComputeFramework/Source/ComputeFramework/Private/ComputeGraph.cpp b/Engine/Plugins/Experimental/ComputeFramework/Source/ComputeFramework/Private/ComputeGraph.cpp
--- a/Engine/Plugins/Experimental/ComputeFramework/Source/ComputeFramework/Private/ComputeGraph.cpp
+++ b/Engine/Plugins/Experimental/ComputeFramework/Source/ComputeFramework/Private/ComputeGraph.cpp
@@ -406,7 +406,7 @@ namespace
 	}
 }
 
-FString UComputeGraph::BuildKernelSource(int32 KernelIndex, FComputeKernelDefinitionSet& OutDefinitionSet, FComputeKernelPermutationVector& OutPermutationVector) const
+FString UComputeGraph::BuildKernelSource(int32 KernelIndex, FString& OutHashKey, FComputeKernelDefinitionSet& OutDefinitionSet, FComputeKernelPermutationVector& OutPermutationVector) const
 {
 	FString HLSL;
 
@@ -446,6 +446,9 @@ FString UComputeGraph::BuildKernelSource(int32 KernelIndex, FComputeKernelDefini
 					// Get define and permutation info for each data provider.
 					DataInterface->GetDefines(OutDefinitionSet);
 					DataInterface->GetPermutations(OutPermutationVector);
+
+					// Accumulate the hash key contribution from the data provider.
+					DataInterface->GetShaderHash(OutHashKey);
 				}
 			}
 
@@ -479,8 +482,12 @@ FString UComputeGraph::BuildKernelSource(int32 KernelIndex, FComputeKernelDefini
 			// Add the kernel code.
 			HLSL += KernelSource->GetSource();
 		}
-		// UE_LOG(LogComputeFramework, Display, TEXT("Kernel Source[%s]:\n%s\n"), *GetPathName(), *HLSL);
 	}
+	
+	// Accumulate the hash key contribution from the HLSL.
+	FSHA1 HashState;
+	HashState.UpdateWithString(*HLSL, HLSL.Len());
+	HashState.Finalize().AppendString(OutHashKey);
 
 	return HLSL;
 }
@@ -499,12 +506,12 @@ void UComputeGraph::CacheResourceShadersForRendering(uint32 CompilationFlags)
 				continue;
 			}
 
+			FString ShaderHashKey;
 			FComputeKernelDefinitionSet ShaderDefinitionSet;
 			FComputeKernelPermutationVector ShaderPermutationVector;
 
 			FString ShaderEntryPoint = Kernel->KernelSource->GetEntryPoint();
-			FString ShaderSource = BuildKernelSource(KernelIndex, ShaderDefinitionSet, ShaderPermutationVector);
-			const uint64 ShaderSourceHash = HashCombine(GetTypeHash(*ShaderSource), Kernel->KernelSource->GetSourceCodeHash());
+			FString ShaderSource = BuildKernelSource(KernelIndex, ShaderHashKey, ShaderDefinitionSet, ShaderPermutationVector);
 			FShaderParametersMetadata* ShaderMetadata = BuildKernelShaderMetadata(KernelIndex);
 
 			const ERHIFeatureLevel::Type CacheFeatureLevel = GMaxRHIFeatureLevel;
@@ -516,8 +523,8 @@ void UComputeGraph::CacheResourceShadersForRendering(uint32 CompilationFlags)
 				CacheFeatureLevel, 
 				GetName(), 
 				ShaderEntryPoint, 
-				MoveTemp(ShaderSource), 
-				ShaderSourceHash, 
+				ShaderHashKey,
+				MoveTemp(ShaderSource),
 				ShaderDefinitionSet,
 				ShaderPermutationVector,
 				ShaderMetadata,
@@ -611,12 +618,12 @@ void UComputeGraph::BeginCacheForCookedPlatformData(ITargetPlatform const* Targe
 
 		if (ShaderFormats.Num() > 0)
 		{
+			FString ShaderHashKey;
 			FComputeKernelDefinitionSet ShaderDefinitionSet;
 			FComputeKernelPermutationVector ShaderPermutationVector;
 
 			FString ShaderEntryPoint = KernelSource->GetEntryPoint();
-			FString ShaderSource = BuildKernelSource(KernelIndex, ShaderDefinitionSet, ShaderPermutationVector);
-			uint64 ShaderSourceHash = FCrc::TypeCrc32(*ShaderSource, 0);
+			FString ShaderSource = BuildKernelSource(KernelIndex, ShaderHashKey, ShaderDefinitionSet, ShaderPermutationVector);
 			FShaderParametersMetadata* ShaderMetadata = BuildKernelShaderMetadata(KernelIndex);
 
 			TArray< TUniquePtr<FComputeKernelResource> >& Resources = KernelResources[KernelIndex].CachedKernelResourcesForCooking.FindOrAdd(TargetPlatform);
@@ -631,8 +638,8 @@ void UComputeGraph::BeginCacheForCookedPlatformData(ITargetPlatform const* Targe
 					TargetFeatureLevel, 
 					GetName(), 
 					ShaderEntryPoint, 
-					ShaderSource, 
-					ShaderSourceHash, 
+					ShaderHashKey,
+					ShaderSource,
 					ShaderDefinitionSet,
 					ShaderPermutationVector,
 					ShaderMetadata,
diff --git a/Engine/Plugins/Experimental/ComputeFramework/Source/ComputeFramework/Private/ComputeKernelShared.cpp b/Engine/Plugins/Experimental/ComputeFramework/Source/ComputeFramework/Private/ComputeKernelShared.cpp
--- a/Engine/Plugins/Experimental/ComputeFramework/Source/ComputeFramework/Private/ComputeKernelShared.cpp
+++ b/Engine/Plugins/Experimental/ComputeFramework/Source/ComputeFramework/Private/ComputeKernelShared.cpp
@@ -203,8 +203,8 @@ void FComputeKernelResource::SetupResource(
 	ERHIFeatureLevel::Type InFeatureLevel,
 	FString const& InFriendlyName,
 	FString const& InShaderEntryPoint,
+	FString const& InShaderHashKey,
 	FString InShaderSource,
-	uint64 InShaderCodeHash,
 	FComputeKernelDefinitionSet& InShaderDefinitionSet,
 	FComputeKernelPermutationVector& InShaderPermutationVector,
 	FShaderParametersMetadata* InShaderMetadata,
@@ -214,8 +214,8 @@ void FComputeKernelResource::SetupResource(
 	FeatureLevel = InFeatureLevel;
 	FriendlyName = InFriendlyName;
 	ShaderEntryPoint = InShaderEntryPoint;
+	ShaderCodeHash = GetTypeHash(InShaderHashKey);
 	ShaderSource = MoveTemp(InShaderSource);
-	ShaderCodeHash = InShaderCodeHash;
 	ShaderDefinitionSet = MoveTemp(InShaderDefinitionSet);
 	ShaderPermutationVector = MoveTemp(InShaderPermutationVector);
 	ShaderMetadata.Reset(InShaderMetadata);
diff --git a/Engine/Plugins/Experimental/ComputeFramework/Source/ComputeFramework/Public/ComputeFramework/ComputeDataInterface.h b/Engine/Plugins/Experimental/ComputeFramework/Source/ComputeFramework/Public/ComputeFramework/ComputeDataInterface.h
--- a/Engine/Plugins/Experimental/ComputeFramework/Source/ComputeFramework/Public/ComputeFramework/ComputeDataInterface.h
+++ b/Engine/Plugins/Experimental/ComputeFramework/Source/ComputeFramework/Public/ComputeFramework/ComputeDataInterface.h
@@ -35,6 +35,8 @@ public:
 	virtual void GetSupportedOutputs(TArray<FShaderFunctionDefinition>& OutFunctions) const {}
 	/** Gather the shader metadata exposed by the data provider payload. */
 	virtual void GetShaderParameters(TCHAR const* UID, FShaderParametersMetadataBuilder& OutBuilder) const {}
+	/** Get a hash that changes on any data interface changes that affect kernel compilation. */
+	virtual void GetShaderHash(FString& InOutKey) const {}
 	/** Gather the shader code for this data provider. */
 	virtual void GetHLSL(FString& OutHLSL) const {}
 	/** Get types of UObject required to instantiate a UComputeDataProvider from this interface. */
diff --git a/Engine/Plugins/Experimental/ComputeFramework/Source/ComputeFramework/Public/ComputeFramework/ComputeGraph.h b/Engine/Plugins/Experimental/ComputeFramework/Source/ComputeFramework/Public/ComputeFramework/ComputeGraph.h
--- a/Engine/Plugins/Experimental/ComputeFramework/Source/ComputeFramework/Public/ComputeFramework/ComputeGraph.h
+++ b/Engine/Plugins/Experimental/ComputeFramework/Source/ComputeFramework/Public/ComputeFramework/ComputeGraph.h
@@ -164,7 +164,7 @@ private:
 
 #if WITH_EDITOR
 	/** Build the HLSL source for a kernel with its linked data interfaces. */
-	FString BuildKernelSource(int32 KernelIndex, FComputeKernelDefinitionSet& OutDefinitionSet, FComputeKernelPermutationVector& OutPermutationVector) const;
+	FString BuildKernelSource(int32 KernelIndex, FString& OutHashKey, FComputeKernelDefinitionSet& OutDefinitionSet, FComputeKernelPermutationVector& OutPermutationVector) const;
 
 	/** Cache shader resources for all kernels in the graph. */
 	void CacheResourceShadersForRendering(uint32 CompilationFlags);
