PR #9132: Fix: Typo in ComputeFramework Plugin (Contributed by [REDACTED_PERSON])
#rnx

diff --git a/Engine/Plugins/Experimental/ComputeFramework/ComputeFramework.uplugin b/Engine/Plugins/Experimental/ComputeFramework/ComputeFramework.uplugin
--- a/Engine/Plugins/Experimental/ComputeFramework/ComputeFramework.uplugin
+++ b/Engine/Plugins/Experimental/ComputeFramework/ComputeFramework.uplugin
@@ -2,7 +2,7 @@
 	"FileVersion": 3,
 	"Version": 1,
 	"VersionName": "0.1",
-	"FriendlyName": "Compute Frameowrk (Experimental)",
+	"FriendlyName": "Compute Framework (Experimental)",
 	"Description": "Support for user authored GPU compute graphs",
 	"Category": "Rendering",
 	"CreatedBy": "Epic Games, Inc.",
