ComputeFramework: Hide raw compute framework assets from create menus.
#rnx

diff --git a/Engine/Plugins/Experimental/ComputeFramework/Source/ComputeFrameworkEditor/Private/ComputeFrameworkEditorModule.cpp b/Engine/Plugins/Experimental/ComputeFramework/Source/ComputeFrameworkEditor/Private/ComputeFrameworkEditorModule.cpp
--- a/Engine/Plugins/Experimental/ComputeFramework/Source/ComputeFrameworkEditor/Private/ComputeFrameworkEditorModule.cpp
+++ b/Engine/Plugins/Experimental/ComputeFramework/Source/ComputeFrameworkEditor/Private/ComputeFrameworkEditorModule.cpp
@@ -12,9 +12,9 @@ void FComputeFrameworkEditorModule::StartupModule()
 {
 	TickObject = MakeUnique<FComputeFrameworkCompilationTick>();
 
-	IAssetTools& AssetTools = FModuleManager::LoadModuleChecked<FAssetToolsModule>("AssetTools").Get();
-	AssetTools.RegisterAssetTypeActions(MakeShareable(new FAssetTypeActions_ComputeKernel));
-	AssetTools.RegisterAssetTypeActions(MakeShareable(new FAssetTypeActions_ComputeKernelFromText));
+// 	IAssetTools& AssetTools = FModuleManager::LoadModuleChecked<FAssetToolsModule>("AssetTools").Get();
+// 	AssetTools.RegisterAssetTypeActions(MakeShareable(new FAssetTypeActions_ComputeKernel));
+// 	AssetTools.RegisterAssetTypeActions(MakeShareable(new FAssetTypeActions_ComputeKernelFromText));
 }
 
 void FComputeFrameworkEditorModule::ShutdownModule()
