ComputeFramework: Move AdditionalSources into Kernel Source.
#rns

diff --git a/Engine/Plugins/Experimental/ComputeFramework/Source/ComputeFramework/Private/ComputeGraph.cpp b/Engine/Plugins/Experimental/ComputeFramework/Source/ComputeFramework/Private/ComputeGraph.cpp
--- a/Engine/Plugins/Experimental/ComputeFramework/Source/ComputeFramework/Private/ComputeGraph.cpp
+++ b/Engine/Plugins/Experimental/ComputeFramework/Source/ComputeFramework/Private/ComputeGraph.cpp
@@ -411,9 +411,7 @@ namespace
 			{
 				if (!InOutUniqueSources.Contains(Source))
 				{
-					TArray<UComputeSource*> AdditionalSources;
-					Source->GetAdditionalSources(AdditionalSources);
-					AddSourcesRecursive(AdditionalSources, InOutUniqueSources);
+					AddSourcesRecursive(Source->AdditionalSources, InOutUniqueSources);
 
 					InOutUniqueSources.AddUnique(Source);
 				}
@@ -550,7 +548,7 @@ void UComputeGraph::CacheResourceShadersForRendering(uint32 CompilationFlags)
 				continue;
 			}
 
-			TMap<FString, FString> AdditionalSources = GatherAdditionalSources(Kernel->AdditionalSources);
+			TMap<FString, FString> AdditionalSources = GatherAdditionalSources(Kernel->KernelSource->AdditionalSources);
 
 			FString ShaderHashKey;
 			TUniquePtr <FComputeKernelDefinitionSet> ShaderDefinitionSet = MakeUnique<FComputeKernelDefinitionSet>();
@@ -666,7 +664,7 @@ void UComputeGraph::BeginCacheForCookedPlatformData(ITargetPlatform const* Targe
 
 		if (ShaderFormats.Num() > 0)
 		{
-			TMap<FString, FString> AdditionalSources = GatherAdditionalSources(KernelInvocations[KernelIndex]->AdditionalSources);
+			TMap<FString, FString> AdditionalSources = GatherAdditionalSources(KernelInvocations[KernelIndex]->KernelSource->AdditionalSources);
 
 			FString ShaderHashKey;
 			TUniquePtr <FComputeKernelDefinitionSet> ShaderDefinitionSet = MakeUnique<FComputeKernelDefinitionSet>();
diff --git a/Engine/Plugins/Experimental/ComputeFramework/Source/ComputeFramework/Private/ComputeKernelFromText.cpp b/Engine/Plugins/Experimental/ComputeFramework/Source/ComputeFramework/Private/ComputeKernelFromText.cpp
--- a/Engine/Plugins/Experimental/ComputeFramework/Source/ComputeFramework/Private/ComputeKernelFromText.cpp
+++ b/Engine/Plugins/Experimental/ComputeFramework/Source/ComputeFramework/Private/ComputeKernelFromText.cpp
@@ -1,12 +1,10 @@
 // Copyright Epic Games, Inc. All Rights Reserved.
 
 #include "ComputeFramework/ComputeKernelFromText.h"
-#include "ComputeFramework/ComputeKernel.h"
 #include "ComputeFramework/ComputeFramework.h"
 #include "Internationalization/Regex.h"
 #include "Misc/FileHelper.h"
 #include "Misc/Paths.h"
-#include "Misc/ScopeExit.h"
 
 #if WITH_EDITOR
 
diff --git a/Engine/Plugins/Experimental/ComputeFramework/Source/ComputeFramework/Private/ComputeSourceFromText.cpp b/Engine/Plugins/Experimental/ComputeFramework/Source/ComputeFramework/Private/ComputeSourceFromText.cpp
--- a/Engine/Plugins/Experimental/ComputeFramework/Source/ComputeFramework/Private/ComputeSourceFromText.cpp
+++ b/Engine/Plugins/Experimental/ComputeFramework/Source/ComputeFramework/Private/ComputeSourceFromText.cpp
@@ -5,12 +5,6 @@
 #include "Internationalization/Regex.h"
 #include "Misc/FileHelper.h"
 #include "Misc/Paths.h"
-#include "Misc/ScopeExit.h"
-
-UComputeSourceFromText::UComputeSourceFromText()
-{
-	UniqueId = FGuid::NewGuid();
-}
 
 #if WITH_EDITOR
 
diff --git a/Engine/Plugins/Experimental/ComputeFramework/Source/ComputeFramework/Public/ComputeFramework/ComputeKernel.h b/Engine/Plugins/Experimental/ComputeFramework/Source/ComputeFramework/Public/ComputeFramework/ComputeKernel.h
--- a/Engine/Plugins/Experimental/ComputeFramework/Source/ComputeFramework/Public/ComputeFramework/ComputeKernel.h
+++ b/Engine/Plugins/Experimental/ComputeFramework/Source/ComputeFramework/Public/ComputeFramework/ComputeKernel.h
@@ -6,7 +6,6 @@
 #include "ComputeKernel.generated.h"
 
 class UComputeKernelSource;
-class UComputeSource;
 
 /** Flags that convey kernel behavior to aid compilation/optimizations. */
 UENUM(meta = (Bitflags, UseEnumValuesAsMaskValuesInEditor = "true"))
@@ -42,10 +41,6 @@ public:
 	UPROPERTY(EditAnywhere, AssetRegistrySearchable, meta = (ShowOnlyInnerProperties), Category = "Kernel")
 	TObjectPtr<UComputeKernelSource> KernelSource = nullptr;
 
-	/** An array of additional independent source assets that the kernel source depends on. */
-	UPROPERTY(EditAnywhere, AssetRegistrySearchable, meta = (ShowOnlyInnerProperties), Category = "Kernel")
-	TArray<TObjectPtr<UComputeSource>> AdditionalSources;
-
 	/** Specifying certain memory access flags allows for optimizations such as kernel fusing. */
 	UPROPERTY(VisibleAnywhere, BlueprintReadOnly, meta = (Bitmask, BitmaskEnum = "/Script/ComputeFramework.EComputeKernelFlags"), Category = "Kernel")
 	int32 KernelFlags = 0;
diff --git a/Engine/Plugins/Experimental/ComputeFramework/Source/ComputeFramework/Public/ComputeFramework/ComputeKernelSource.h b/Engine/Plugins/Experimental/ComputeFramework/Source/ComputeFramework/Public/ComputeFramework/ComputeKernelSource.h
--- a/Engine/Plugins/Experimental/ComputeFramework/Source/ComputeFramework/Public/ComputeFramework/ComputeKernelSource.h
+++ b/Engine/Plugins/Experimental/ComputeFramework/Source/ComputeFramework/Public/ComputeFramework/ComputeKernelSource.h
@@ -6,6 +6,8 @@
 #include "ComputeFramework/ShaderParamTypeDefinition.h"
 #include "ComputeKernelSource.generated.h"
 
+class UComputeSource;
+
 /** 
  * Class representing the source for a UComputeKernel 
  * We derive from this for each authoring mechanism. (HLSL text, VPL graph, ML Meta Lang, etc.)
@@ -15,13 +17,17 @@ class COMPUTEFRAMEWORK_API UComputeKernelSource : public UObject
 {
 	GENERATED_BODY()
 
+public:
+	/** Get kernel source code ready for HLSL compilation. */
+	virtual FString GetSource() const PURE_VIRTUAL(UComputeKernelSource::GetSource, return {};);
+
 public:
 	/** Kernel entry point. */
 	UPROPERTY(VisibleAnywhere, AssetRegistrySearchable, Category = "Kernel")
 	FString EntryPoint;
 
 	/** Kernel group size. */
-	UPROPERTY(VisibleAnywhere, Category = "Kernel")
+	UPROPERTY(VisibleAnywhere, AssetRegistrySearchable, Category = "Kernel")
 	FIntVector GroupSize = FIntVector(64, 1, 1);
 
 	/** Base permutations exposed by the kernel. These will be extended by further permutations declared in any linked data providers. */
@@ -32,6 +38,10 @@ public:
 	UPROPERTY(VisibleAnywhere, meta = (ShowOnlyInnerProperties), Category = "Kernel")
 	FComputeKernelDefinitionSet DefinitionsSet;
 
+	/** An array of additional independent source assets that the kernel source depends on. */
+	UPROPERTY(EditAnywhere, meta = (ShowOnlyInnerProperties), Category = "External")
+	TArray<TObjectPtr<UComputeSource>> AdditionalSources;
+
 	/* Named external inputs for the kernel. These must be fulfilled by linked data providers. */
 	UPROPERTY(VisibleAnywhere, EditFixedSize, Category = "External")
 	TArray<FShaderFunctionDefinition> ExternalInputs;
@@ -39,7 +49,4 @@ public:
 	/* Named external outputs for the kernel. These must be fulfilled by linked data providers. */
 	UPROPERTY(VisibleAnywhere, EditFixedSize, Category = "External")
 	TArray<FShaderFunctionDefinition> ExternalOutputs;
-
-	/** Get kernel source code. */
-	virtual FString GetSource() const PURE_VIRTUAL(UComputeKernelSource::GetSource, return {};);
 };
diff --git a/Engine/Plugins/Experimental/ComputeFramework/Source/ComputeFramework/Public/ComputeFramework/ComputeSource.h b/Engine/Plugins/Experimental/ComputeFramework/Source/ComputeFramework/Public/ComputeFramework/ComputeSource.h
--- a/Engine/Plugins/Experimental/ComputeFramework/Source/ComputeFramework/Public/ComputeFramework/ComputeSource.h
+++ b/Engine/Plugins/Experimental/ComputeFramework/Source/ComputeFramework/Public/ComputeFramework/ComputeSource.h
@@ -7,7 +7,7 @@
 
 /** 
  * Class representing some source for inclusion in a UComputeKernel.
- * Source can be created by different authoring mechanisms. (HLSL text, VPL graph, ML Meta Lang, etc.)
+ * We derive from this for each authoring mechanism. (HLSL text, VPL graph, ML Meta Lang, etc.)
  */
 UCLASS(Abstract)
 class COMPUTEFRAMEWORK_API UComputeSource : public UObject
@@ -15,8 +15,11 @@ class COMPUTEFRAMEWORK_API UComputeSource : public UObject
 	GENERATED_BODY()
 
 public:
-	/** Get source code. */
+	/** Get source code ready for HLSL compilation. */
 	virtual FString GetSource() const { return FString(); }
-	/** Get an array of source objects. This allows us to specify source dependencies. */
-	virtual void GetAdditionalSources(TArray<UComputeSource*>& OutSources) const {};
+
+public:
+	/** Array of additional source objects. This allows us to specify source dependencies. */
+	UPROPERTY(EditAnywhere, Category = Source)
+	TArray<TObjectPtr<UComputeSource>> AdditionalSources;
 };
diff --git a/Engine/Plugins/Experimental/ComputeFramework/Source/ComputeFramework/Public/ComputeFramework/ComputeSourceFromText.h b/Engine/Plugins/Experimental/ComputeFramework/Source/ComputeFramework/Public/ComputeFramework/ComputeSourceFromText.h
--- a/Engine/Plugins/Experimental/ComputeFramework/Source/ComputeFramework/Public/ComputeFramework/ComputeSourceFromText.h
+++ b/Engine/Plugins/Experimental/ComputeFramework/Source/ComputeFramework/Public/ComputeFramework/ComputeSourceFromText.h
@@ -15,16 +15,10 @@ class COMPUTEFRAMEWORK_API UComputeSourceFromText : public UComputeSource
 	GENERATED_BODY()
 
 public:
-	UComputeSourceFromText();
-
 	/** Filepath to the source file containing the kernel entry points and all options for parsing. */
 	UPROPERTY(EditDefaultsOnly, AssetRegistrySearchable, meta = (ContentDir, RelativeToGameContentDir, FilePathFilter = "Unreal Shader File (*.usf)|*.usf"), Category = "Kernel")
 	FFilePath SourceFile;
 
-	/** A unique id for the asset. */
-	UPROPERTY()
-	FGuid UniqueId;
-
 #if WITH_EDITOR
 	/** Parse the source to get metadata. */
 	void ReparseSourceText();
