ComputeFramework: Add compilation completion message to outputs of kernel compiler.
Use source object path for virtual source include paths by default.
#rnx

diff --git a/Engine/Plugins/Experimental/ComputeFramework/Source/ComputeFramework/Private/ComputeFramework/ComputeKernelShared.h b/Engine/Plugins/Experimental/ComputeFramework/Source/ComputeFramework/Private/ComputeFramework/ComputeKernelShared.h
--- a/Engine/Plugins/Experimental/ComputeFramework/Source/ComputeFramework/Private/ComputeFramework/ComputeKernelShared.h
+++ b/Engine/Plugins/Experimental/ComputeFramework/Source/ComputeFramework/Private/ComputeFramework/ComputeKernelShared.h
@@ -280,7 +280,7 @@ private:
 };
 
 /** Delegate for kernel compile completion. */
-DECLARE_DELEGATE_OneParam(FOnComputeKernelCompilationComplete, class FComputeKernelResource const*);
+DECLARE_DELEGATE_OneParam(FOnComputeKernelCompilationComplete, FComputeKernelResource const*);
 
 /**
  * FComputeKernelResource represents a UComputeKernel to the shader compilation process
@@ -329,7 +329,7 @@ public:
 	/**
 	 * Called when compilation finishes, after the GameThreadShaderMap is set and the render command to set the RenderThreadShaderMap is queued
 	 */
-	virtual void NotifyCompilationFinished();
+	virtual void NotifyCompilationFinished(FString const& ResultMessage);
 
 	/**
 	 * Cancels all outstanding compilation jobs
@@ -349,8 +349,8 @@ public:
 	bool IsCompilationFinished() const;
 
 	// Accessors.
-	const TArray<FString>& GetCompileErrors() const { return CompileErrors; }
-	void SetCompileErrors(const TArray<FString>& InCompileErrors) { CompileErrors = InCompileErrors; }
+	const TArray<FString>& GetCompileOutputMessages() const { return CompileOutputMessages; }
+	void SetCompileOutputMessages(TArray<FString> const& InCompileOutputMessages) { CompileOutputMessages = InCompileOutputMessages; }
 
 	ERHIFeatureLevel::Type GetFeatureLevel() const { return FeatureLevel; }
 
@@ -427,11 +427,6 @@ public:
 
 	void SetupCompileEnvironment(int32 InPermutationId, FShaderCompilerEnvironment& OutShaderEnvironment) const;
 
-	void SetCompileErrors(TArray<FString> &InErrors)
-	{
-		CompileErrors = InErrors;
-	}
-
 	const FShaderParametersMetadata* GetShaderParamMetadata() const
 	{
 		return ShaderParameterMetadata;
@@ -454,7 +449,7 @@ protected:
 	void GetShaderMapIDsWithUnfinishedCompilation(TArray<int32>& OutShaderMapIds);
 
 private:
-	TArray<FString> CompileErrors;
+	TArray<FString> CompileOutputMessages;
 
 	/** 
 	 * Game thread tracked shader map, which is ref counted and manages shader map lifetime. 
diff --git a/Engine/Plugins/Experimental/ComputeFramework/Source/ComputeFramework/Private/ComputeGraph.cpp b/Engine/Plugins/Experimental/ComputeFramework/Source/ComputeFramework/Private/ComputeGraph.cpp
--- a/Engine/Plugins/Experimental/ComputeFramework/Source/ComputeFramework/Private/ComputeGraph.cpp
+++ b/Engine/Plugins/Experimental/ComputeFramework/Source/ComputeFramework/Private/ComputeGraph.cpp
@@ -429,9 +429,7 @@ namespace
 
 		for (UComputeSource const* Source : UniqueSources)
 		{
-			Result.Add(TPair<FString, FString>(
-				FString::Printf(TEXT("/Engine/Generated/%s.ush"), *Source->GetName()),
-				Source->GetSource()));
+			Result.Add(Source->GetVirtualPath(), Source->GetSource());
 		}
 
 		return Result;
@@ -556,6 +554,7 @@ void UComputeGraph::CacheResourceShadersForRendering(uint32 CompilationFlags)
 			TUniquePtr <FShaderParametersMetadataAllocations> ShaderParameterMetadataAllocations = MakeUnique<FShaderParametersMetadataAllocations>();
 
 			FString ShaderEntryPoint = Kernel->KernelSource->EntryPoint;
+			FString ShaderFriendlyName = GetOuter()->GetName() + TEXT("_") + ShaderEntryPoint;
 			FString ShaderSource = BuildKernelSource(KernelIndex, *Kernel->KernelSource, AdditionalSources, ShaderHashKey, *ShaderDefinitionSet, *ShaderPermutationVector);
 			FShaderParametersMetadata* ShaderParameterMetadata = BuildKernelShaderMetadata(KernelIndex, *ShaderParameterMetadataAllocations);
 
@@ -566,7 +565,7 @@ void UComputeGraph::CacheResourceShadersForRendering(uint32 CompilationFlags)
 			// Now we have all the information that the KernelResource will need for compilation.
 			KernelResource->SetupResource(
 				CacheFeatureLevel, 
-				GetName(), 
+				ShaderFriendlyName,
 				ShaderEntryPoint, 
 				ShaderHashKey,
 				ShaderSource,
@@ -626,11 +625,9 @@ void UComputeGraph::CacheShadersForResource(
 			*LegacyShaderPlatformToShaderFormat(ShaderPlatform).ToString()
 		);
 
-		auto& CompilationErrors = KernelResource->GetCompileErrors();
-		uint32 ErrorCount = CompilationErrors.Num();
-		for (uint32 i = 0; i < ErrorCount; ++i)
+		for (FString const& Message : KernelResource->GetCompileOutputMessages())
 		{
-			UE_LOG(LogComputeFramework, Warning, TEXT("      [Error] - %s"), *CompilationErrors[i]);
+			UE_LOG(LogComputeFramework, Warning, TEXT("   %s"), *Message);
 		}
 	}
 }
@@ -642,7 +639,7 @@ void UComputeGraph::ShaderCompileCompletionCallback(FComputeKernelResource const
 	{
 		if (KernelResource == KernelResources[KernelIndex].Get())
 		{
-			OnKernelCompilationComplete(KernelIndex, KernelResource->GetCompileErrors());
+			OnKernelCompilationComplete(KernelIndex, KernelResource->GetCompileOutputMessages());
 		}
 	}
 }
@@ -671,6 +668,7 @@ void UComputeGraph::BeginCacheForCookedPlatformData(ITargetPlatform const* Targe
 			TUniquePtr <FComputeKernelPermutationVector> ShaderPermutationVector = MakeUnique<FComputeKernelPermutationVector>();
 
 			FString ShaderEntryPoint = KernelSource->EntryPoint;
+			FString ShaderFriendlyName = GetOuter()->GetName() + TEXT("_") + ShaderEntryPoint;
 			FString ShaderSource = BuildKernelSource(KernelIndex, *KernelSource, AdditionalSources, ShaderHashKey, *ShaderDefinitionSet, *ShaderPermutationVector);
 
 			TArray< TUniquePtr<FComputeKernelResource> >& Resources = KernelResources[KernelIndex].CachedKernelResourcesForCooking.FindOrAdd(TargetPlatform);
@@ -686,7 +684,7 @@ void UComputeGraph::BeginCacheForCookedPlatformData(ITargetPlatform const* Targe
 				TUniquePtr<FComputeKernelResource> KernelResource = MakeUnique<FComputeKernelResource>();
 				KernelResource->SetupResource(
 					TargetFeatureLevel, 
-					GetName(), 
+					ShaderFriendlyName,
 					ShaderEntryPoint, 
 					ShaderHashKey,
 					ShaderSource,
diff --git a/Engine/Plugins/Experimental/ComputeFramework/Source/ComputeFramework/Private/ComputeKernelShaderCompilationManager.cpp b/Engine/Plugins/Experimental/ComputeFramework/Source/ComputeFramework/Private/ComputeKernelShaderCompilationManager.cpp
--- a/Engine/Plugins/Experimental/ComputeFramework/Source/ComputeFramework/Private/ComputeKernelShaderCompilationManager.cpp
+++ b/Engine/Plugins/Experimental/ComputeFramework/Source/ComputeFramework/Private/ComputeKernelShaderCompilationManager.cpp
@@ -15,7 +15,7 @@
 
 DEFINE_LOG_CATEGORY_STATIC(LogComputeKernelShaderCompiler, All, All);
 
-static int32 GShowComputeKernelShaderWarnings = 1;
+static int32 GShowComputeKernelShaderWarnings = 0;
 static FAutoConsoleVariableRef CVarShowComputeKernelShaderWarnings(
 	TEXT("ComputeKernel.ShowShaderCompilerWarnings"),
 	GShowComputeKernelShaderWarnings,
@@ -319,7 +319,6 @@ void FComputeKernelShaderCompilationManager::ProcessCompiledComputeKernelShaderM
 				bShaderMapComplete = ShaderMap->ProcessCompilationResults(ResultArray, CompileResults.FinalizeJobIndex, TimeBudget);
 			}
 
-
 			if (bShaderMapComplete)
 			{
 				ShaderMap->SetCompiledSuccessfully(bSuccess);
@@ -356,7 +355,7 @@ void FComputeKernelShaderCompilationManager::ProcessCompiledComputeKernelShaderM
 						if (!bSuccess)
 						{
 							// Propagate error messages
-							Kernel->SetCompileErrors(Errors);
+							Kernel->SetCompileOutputMessages(Errors);
 							KernelsToUpdate.Add(Kernel, nullptr);
 
 							for (int32 ErrorIndex = 0; ErrorIndex < Errors.Num(); ErrorIndex++)
@@ -386,6 +385,8 @@ void FComputeKernelShaderCompilationManager::ProcessCompiledComputeKernelShaderM
 								{
 									UE_LOG(LogComputeKernelShaderCompiler, Warning, TEXT("	%s"), *Errors[ErrorIndex]);
 								}
+
+								Kernel->SetCompileOutputMessages(Errors);
 							}
 						}
 					}
@@ -393,7 +394,10 @@ void FComputeKernelShaderCompilationManager::ProcessCompiledComputeKernelShaderM
 					{
 						if (CompletedShaderMap->IsComplete(Kernel, true))
 						{
-							Kernel->NotifyCompilationFinished();
+							const FString SuccessMessage = FString::Printf(TEXT("%s: Shader compilation success!"), *Kernel->GetFriendlyName());
+							const FString FailMessage = FString::Printf(TEXT("%s: Shader compilation failed."), *Kernel->GetFriendlyName());
+
+							Kernel->NotifyCompilationFinished(bSuccess ? SuccessMessage : FailMessage);
 						}
 					}
 				}
@@ -425,8 +429,11 @@ void FComputeKernelShaderCompilationManager::ProcessCompiledComputeKernelShaderM
 					Kernel->SetRenderingThreadShaderMap(ShaderMap);
 				});
 
+			const bool bSuccess = ShaderMap && ShaderMap->CompiledSuccessfully();
+			const FString SuccessMessage = FString::Printf(TEXT("%s: Shader compilation success!"), *Kernel->GetFriendlyName());
+			const FString FailMessage = FString::Printf(TEXT("%s: Shader compilation failed."), *Kernel->GetFriendlyName());
 
-			Kernel->NotifyCompilationFinished();
+			Kernel->NotifyCompilationFinished(bSuccess ? SuccessMessage : FailMessage);
 		}
 	}
 #endif
diff --git a/Engine/Plugins/Experimental/ComputeFramework/Source/ComputeFramework/Private/ComputeKernelShaderMap.cpp b/Engine/Plugins/Experimental/ComputeFramework/Source/ComputeFramework/Private/ComputeKernelShaderMap.cpp
--- a/Engine/Plugins/Experimental/ComputeFramework/Source/ComputeFramework/Private/ComputeKernelShaderMap.cpp
+++ b/Engine/Plugins/Experimental/ComputeFramework/Source/ComputeFramework/Private/ComputeKernelShaderMap.cpp
@@ -428,9 +428,12 @@ void FComputeKernelShaderMap::Compile(
 				}
 				else if (ShaderType)
 				{
-					UE_LOG(LogShaders, Display, TEXT("Skipping compilation of %s as it isn't supported on this target type."), *InKernel->GetFriendlyName());
 					InKernel->RemoveOutstandingCompileId(CompilingId);
-					InKernel->NotifyCompilationFinished();
+
+					FString Message = FString::Printf(TEXT("%s: Skipping compilation because it isn't supported on this target type."), *InKernel->GetFriendlyName());
+					UE_LOG(LogShaders, Display, TEXT("%s"), *Message);
+					
+					InKernel->NotifyCompilationFinished(Message);
 				}
 			}
   
@@ -717,7 +720,7 @@ void FComputeKernelShaderMap::RemovePending(FComputeKernelResource* InKernel)
 		if (Result)
 		{
 			InKernel->RemoveOutstandingCompileId(It.Key()->CompilingId);
-			InKernel->NotifyCompilationFinished();
+			InKernel->NotifyCompilationFinished(FString::Printf(TEXT("%s: Removing compilation."), *InKernel->GetFriendlyName()));
 		}
 #if DEBUG_INFINITESHADERCOMPILE
 		if ( Result )
diff --git a/Engine/Plugins/Experimental/ComputeFramework/Source/ComputeFramework/Private/ComputeKernelShared.cpp b/Engine/Plugins/Experimental/ComputeFramework/Source/ComputeFramework/Private/ComputeKernelShared.cpp
--- a/Engine/Plugins/Experimental/ComputeFramework/Source/ComputeFramework/Private/ComputeKernelShared.cpp
+++ b/Engine/Plugins/Experimental/ComputeFramework/Source/ComputeFramework/Private/ComputeKernelShared.cpp
@@ -51,8 +51,9 @@ bool FComputeKernelResource::ShouldCache(EShaderPlatform InPlatform, const FShad
 	return true;
 }
 
-void FComputeKernelResource::NotifyCompilationFinished()
+void FComputeKernelResource::NotifyCompilationFinished(FString const& ResultMessage)
 {
+	CompileOutputMessages.Add(ResultMessage);
 	OnCompilationCompleteDelegate.ExecuteIfBound(this);
 }
 
@@ -62,9 +63,12 @@ void FComputeKernelResource::CancelCompilation()
 	if (IsInGameThread())
 	{
 		FComputeKernelShaderMap::RemovePending(this);
-
-		UE_LOG(LogShaders, Log, TEXT("CancelCompilation %p."), this);
 		OutstandingCompileShaderMapIds.Empty();
+
+		FString Message = FString::Printf(TEXT("%s: Compilation cancelled."), *GetFriendlyName());
+		UE_LOG(LogShaders, Log, TEXT("%s"), *Message);
+
+		NotifyCompilationFinished(Message);
 	}
 #endif
 }
@@ -232,7 +236,7 @@ void FComputeKernelResource::SetupResource(
 	ShaderPermutationVector = MoveTemp(InShaderPermutationVector);
 	ShaderParameterMetadataAllocations = MoveTemp(InShaderParameterMetadataAllocations);
 	ShaderParameterMetadata = InShaderParameterMetadata;
-	CompileErrors.Reset();
+	CompileOutputMessages.Reset();
 #if WITH_EDITOR
 	AssetPath = InAssetPath;
 #endif
@@ -321,7 +325,9 @@ bool FComputeKernelResource::CacheShaders(const FComputeKernelShaderMapId& InSha
 				FComputeKernelShaderMap::LoadFromDerivedDataCache(this, InShaderMapId, InPlatform, GameThreadShaderMap);
 				if (GameThreadShaderMap && GameThreadShaderMap->IsValid())
 				{
-					UE_LOG(LogTemp, Display, TEXT("Loaded shader %s for kernel %s from DDC"), *GameThreadShaderMap->GetFriendlyName(), *GetFriendlyName());
+					FString Message = FString::Printf(TEXT("Loaded shaders for %s from DDC."), *GetFriendlyName());
+					UE_LOG(LogTemp, Display, TEXT("%s"), *Message);
+					CompileOutputMessages.Add(Message);
 				}
 				else
 				{
@@ -370,6 +376,7 @@ bool FComputeKernelResource::CacheShaders(const FComputeKernelShaderMapId& InSha
 	}
 	else
 	{
+		NotifyCompilationFinished(FString::Printf(TEXT("%s: Loaded shaders from DDC."), *GetFriendlyName()));
 		bSucceeded = true;
 	}
 
diff --git a/Engine/Plugins/Experimental/ComputeFramework/Source/ComputeFramework/Public/ComputeFramework/ComputeGraph.h b/Engine/Plugins/Experimental/ComputeFramework/Source/ComputeFramework/Public/ComputeFramework/ComputeGraph.h
--- a/Engine/Plugins/Experimental/ComputeFramework/Source/ComputeFramework/Public/ComputeFramework/ComputeGraph.h
+++ b/Engine/Plugins/Experimental/ComputeFramework/Source/ComputeFramework/Public/ComputeFramework/ComputeGraph.h
@@ -92,7 +92,7 @@ public:
 	virtual ~UComputeGraph();
 
 	/** Called each time that a single kernel shader compilation is completed. */
-	virtual void OnKernelCompilationComplete(int32 InKernelIndex, const TArray<FString>& InCompileErrors) {}
+	virtual void OnKernelCompilationComplete(int32 InKernelIndex, const TArray<FString>& InCompileOutputMessages) {}
 
 	/** 
 	 * Returns true if graph is valid. 
diff --git a/Engine/Plugins/Experimental/ComputeFramework/Source/ComputeFramework/Public/ComputeFramework/ComputeSource.h b/Engine/Plugins/Experimental/ComputeFramework/Source/ComputeFramework/Public/ComputeFramework/ComputeSource.h
--- a/Engine/Plugins/Experimental/ComputeFramework/Source/ComputeFramework/Public/ComputeFramework/ComputeSource.h
+++ b/Engine/Plugins/Experimental/ComputeFramework/Source/ComputeFramework/Public/ComputeFramework/ComputeSource.h
@@ -2,7 +2,7 @@
 
 #pragma once
 
-#include "ComputeFramework/ShaderParamTypeDefinition.h"
+#include "CoreMinimal.h"
 #include "ComputeSource.generated.h"
 
 /** 
@@ -17,6 +17,8 @@ class COMPUTEFRAMEWORK_API UComputeSource : public UObject
 public:
 	/** Get source code ready for HLSL compilation. */
 	virtual FString GetSource() const { return FString(); }
+	/** Get virtual file path for the source code. This will be the file name shown in any compilation errors. */
+	virtual FString GetVirtualPath() const { return GetPathName(); }
 
 public:
 	/** Array of additional source objects. This allows us to specify source dependencies. */
