ComputeFramework: Fix shader build fail outputs going to wrong directory.
Now goes to Saved/ShaderDebugInfo

diff --git a/Engine/Plugins/Experimental/ComputeFramework/Source/ComputeFramework/Private/ComputeKernelShaderCompilationManager.cpp b/Engine/Plugins/Experimental/ComputeFramework/Source/ComputeFramework/Private/ComputeKernelShaderCompilationManager.cpp
--- a/Engine/Plugins/Experimental/ComputeFramework/Source/ComputeFramework/Private/ComputeKernelShaderCompilationManager.cpp
+++ b/Engine/Plugins/Experimental/ComputeFramework/Source/ComputeFramework/Private/ComputeKernelShaderCompilationManager.cpp
@@ -119,7 +119,6 @@ void FComputeKernelShaderCompilationManager::RunCompileJobs()
 				FPaths::NormalizeDirectoryName(CurrentJob.Input.DumpDebugInfoRootPath);
 				const FShaderCompilingManager::EDumpShaderDebugInfo DumpShaderDebugInfo = GShaderCompilingManager->GetDumpShaderDebugInfo();
 				CurrentJob.Input.DebugExtension.Empty();
-				CurrentJob.Input.DumpDebugInfoRootPath.Empty();
 				if (DumpShaderDebugInfo == FShaderCompilingManager::EDumpShaderDebugInfo::Always)
 				{
 					CurrentJob.Input.DumpDebugInfoRootPath = GShaderCompilingManager->CreateShaderDebugInfoPath(CurrentJob.Input);
