ComputeGraph: Fix potential nullptr deref.
#rnx

diff --git a/Engine/Plugins/Experimental/ComputeFramework/Source/ComputeFramework/Private/ComputeGraph.cpp b/Engine/Plugins/Experimental/ComputeFramework/Source/ComputeFramework/Private/ComputeGraph.cpp
--- a/Engine/Plugins/Experimental/ComputeFramework/Source/ComputeFramework/Private/ComputeGraph.cpp
+++ b/Engine/Plugins/Experimental/ComputeFramework/Source/ComputeFramework/Private/ComputeGraph.cpp
@@ -165,7 +165,7 @@ void UComputeGraph::CreateDataProviders(int32 InBindingIndex, TObjectPtr<UObject
 {
 	InOutDataProviders.SetNumZeroed(DataInterfaces.Num());
 
-	if (ensure(Bindings.IsValidIndex(InBindingIndex)) && InBindingObject.IsA(Bindings[InBindingIndex]))
+	if (ensure(Bindings.IsValidIndex(InBindingIndex)) && !InBindingObject.IsNull() && InBindingObject.IsA(Bindings[InBindingIndex]))
 	{
 		for (int32 DataInterfaceIndex = 0; DataInterfaceIndex < DataInterfaces.Num(); ++DataInterfaceIndex)
 		{
