ComputeFramework: During cooking compilation there can be multiple owners of define and permutation set objects. Should use shared instead of unique pointers for these.

diff --git a/Engine/Plugins/Experimental/ComputeFramework/Source/ComputeFramework/Private/ComputeFramework/ComputeKernelShared.h b/Engine/Plugins/Experimental/ComputeFramework/Source/ComputeFramework/Private/ComputeFramework/ComputeKernelShared.h
--- a/Engine/Plugins/Experimental/ComputeFramework/Source/ComputeFramework/Private/ComputeFramework/ComputeKernelShared.h
+++ b/Engine/Plugins/Experimental/ComputeFramework/Source/ComputeFramework/Private/ComputeFramework/ComputeKernelShared.h
@@ -418,8 +418,8 @@ public:
 		FString const& InShaderHashKey,
 		FString& InShaderSource,
 		TMap<FString, FString>& InAdditionalSources,
-		TUniquePtr<FComputeKernelDefinitionSet>& InShaderDefinitionSet,
-		TUniquePtr<FComputeKernelPermutationVector>& InShaderPermutationVector,
+		TSharedPtr<FComputeKernelDefinitionSet>& InShaderDefinitionSet,
+		TSharedPtr<FComputeKernelPermutationVector>& InShaderPermutationVector,
 		TUniquePtr<FShaderParametersMetadataAllocations>& InShaderParameterMetadataAllocations,
 		FShaderParametersMetadata* InShaderParameterMetadata,
 		FName const& InAssetPath
@@ -480,11 +480,11 @@ private:
 	/** Additional source code. Stored as a map from a virtual source file name to the shader source.  */
 	TMap<FString, FString> AdditionalSources;
 
-	/** Defines used when compiling shaders. */
-	TUniquePtr<FComputeKernelDefinitionSet> ShaderDefinitionSet;
+	/** Defines used when compiling shaders. Object can be shared across a number of shader formats during cook. */
+	TSharedPtr<FComputeKernelDefinitionSet> ShaderDefinitionSet;
 	
-	/** Permutations used for compiling shaders. */
-	TUniquePtr<FComputeKernelPermutationVector> ShaderPermutationVector;
+	/** Permutations used for compiling shaders. Object can be shared across a number of shader formats during cook. */
+	TSharedPtr<FComputeKernelPermutationVector> ShaderPermutationVector;
 
 	/** Allocations for the ShaderParameterMetadata */
 	TUniquePtr<FShaderParametersMetadataAllocations> ShaderParameterMetadataAllocations;
diff --git a/Engine/Plugins/Experimental/ComputeFramework/Source/ComputeFramework/Private/ComputeGraph.cpp b/Engine/Plugins/Experimental/ComputeFramework/Source/ComputeFramework/Private/ComputeGraph.cpp
--- a/Engine/Plugins/Experimental/ComputeFramework/Source/ComputeFramework/Private/ComputeGraph.cpp
+++ b/Engine/Plugins/Experimental/ComputeFramework/Source/ComputeFramework/Private/ComputeGraph.cpp
@@ -562,9 +562,9 @@ void UComputeGraph::CacheResourceShadersForRendering(uint32 CompilationFlags)
 			TMap<FString, FString> AdditionalSources = GatherAdditionalSources(Kernel->KernelSource->AdditionalSources);
 
 			FString ShaderHashKey;
-			TUniquePtr <FComputeKernelDefinitionSet> ShaderDefinitionSet = MakeUnique<FComputeKernelDefinitionSet>();
-			TUniquePtr <FComputeKernelPermutationVector> ShaderPermutationVector = MakeUnique<FComputeKernelPermutationVector>();
-			TUniquePtr <FShaderParametersMetadataAllocations> ShaderParameterMetadataAllocations = MakeUnique<FShaderParametersMetadataAllocations>();
+			TSharedPtr<FComputeKernelDefinitionSet> ShaderDefinitionSet = MakeShared<FComputeKernelDefinitionSet>();
+			TSharedPtr<FComputeKernelPermutationVector> ShaderPermutationVector = MakeShared<FComputeKernelPermutationVector>();
+			TUniquePtr<FShaderParametersMetadataAllocations> ShaderParameterMetadataAllocations = MakeUnique<FShaderParametersMetadataAllocations>();
 
 			FString ShaderEntryPoint = Kernel->KernelSource->EntryPoint;
 			FString ShaderFriendlyName = GetOuter()->GetName() + TEXT("_") + ShaderEntryPoint;
@@ -677,8 +677,8 @@ void UComputeGraph::BeginCacheForCookedPlatformData(ITargetPlatform const* Targe
 			TMap<FString, FString> AdditionalSources = GatherAdditionalSources(KernelInvocations[KernelIndex]->KernelSource->AdditionalSources);
 
 			FString ShaderHashKey;
-			TUniquePtr <FComputeKernelDefinitionSet> ShaderDefinitionSet = MakeUnique<FComputeKernelDefinitionSet>();
-			TUniquePtr <FComputeKernelPermutationVector> ShaderPermutationVector = MakeUnique<FComputeKernelPermutationVector>();
+			TSharedPtr<FComputeKernelDefinitionSet> ShaderDefinitionSet = MakeShared<FComputeKernelDefinitionSet>();
+			TSharedPtr<FComputeKernelPermutationVector> ShaderPermutationVector = MakeShared<FComputeKernelPermutationVector>();
 
 			FString ShaderEntryPoint = KernelSource->EntryPoint;
 			FString ShaderFriendlyName = GetOuter()->GetName() + TEXT("_") + ShaderEntryPoint;
diff --git a/Engine/Plugins/Experimental/ComputeFramework/Source/ComputeFramework/Private/ComputeKernelShared.cpp b/Engine/Plugins/Experimental/ComputeFramework/Source/ComputeFramework/Private/ComputeKernelShared.cpp
--- a/Engine/Plugins/Experimental/ComputeFramework/Source/ComputeFramework/Private/ComputeKernelShared.cpp
+++ b/Engine/Plugins/Experimental/ComputeFramework/Source/ComputeFramework/Private/ComputeKernelShared.cpp
@@ -219,8 +219,8 @@ void FComputeKernelResource::SetupResource(
 	FString const& InShaderHashKey,
 	FString& InShaderSource,
 	TMap<FString, FString>& InAdditionalSources,
-	TUniquePtr<FComputeKernelDefinitionSet>& InShaderDefinitionSet,
-	TUniquePtr<FComputeKernelPermutationVector>& InShaderPermutationVector,
+	TSharedPtr<FComputeKernelDefinitionSet>& InShaderDefinitionSet,
+	TSharedPtr<FComputeKernelPermutationVector>& InShaderPermutationVector,
 	TUniquePtr<FShaderParametersMetadataAllocations>& InShaderParameterMetadataAllocations,
 	FShaderParametersMetadata* InShaderParameterMetadata,
 	FName const& InAssetPath
@@ -232,8 +232,8 @@ void FComputeKernelResource::SetupResource(
 	ShaderCodeHash = GetTypeHash(InShaderHashKey);
 	ShaderSource = MoveTemp(InShaderSource);
 	AdditionalSources = MoveTemp(InAdditionalSources);
-	ShaderDefinitionSet = MoveTemp(InShaderDefinitionSet);
-	ShaderPermutationVector = MoveTemp(InShaderPermutationVector);
+	ShaderDefinitionSet = InShaderDefinitionSet;
+	ShaderPermutationVector = InShaderPermutationVector;
 	ShaderParameterMetadataAllocations = MoveTemp(InShaderParameterMetadataAllocations);
 	ShaderParameterMetadata = InShaderParameterMetadata;
 	CompileOutputMessages.Reset();
