ComputeFramework: Don't move source string during graph compilation during cook, since it is then not available after the first permutation.
#rnx

diff --git a/Engine/Plugins/Experimental/ComputeFramework/Source/ComputeFramework/Private/ComputeFramework/ComputeKernelShared.h b/Engine/Plugins/Experimental/ComputeFramework/Source/ComputeFramework/Private/ComputeFramework/ComputeKernelShared.h
--- a/Engine/Plugins/Experimental/ComputeFramework/Source/ComputeFramework/Private/ComputeFramework/ComputeKernelShared.h
+++ b/Engine/Plugins/Experimental/ComputeFramework/Source/ComputeFramework/Private/ComputeFramework/ComputeKernelShared.h
@@ -416,8 +416,8 @@ public:
 		FString const& InFriendlyName,
 		FString const& InShaderEntryPoint,
 		FString const& InShaderHashKey,
-		FString& InShaderSource,
-		TMap<FString, FString>& InAdditionalSources,
+		FString const& InShaderSource,
+		TMap<FString, FString> const& InAdditionalSources,
 		TSharedPtr<FComputeKernelDefinitionSet>& InShaderDefinitionSet,
 		TSharedPtr<FComputeKernelPermutationVector>& InShaderPermutationVector,
 		TUniquePtr<FShaderParametersMetadataAllocations>& InShaderParameterMetadataAllocations,
diff --git a/Engine/Plugins/Experimental/ComputeFramework/Source/ComputeFramework/Private/ComputeKernelShared.cpp b/Engine/Plugins/Experimental/ComputeFramework/Source/ComputeFramework/Private/ComputeKernelShared.cpp
--- a/Engine/Plugins/Experimental/ComputeFramework/Source/ComputeFramework/Private/ComputeKernelShared.cpp
+++ b/Engine/Plugins/Experimental/ComputeFramework/Source/ComputeFramework/Private/ComputeKernelShared.cpp
@@ -217,8 +217,8 @@ void FComputeKernelResource::SetupResource(
 	FString const& InFriendlyName,
 	FString const& InShaderEntryPoint,
 	FString const& InShaderHashKey,
-	FString& InShaderSource,
-	TMap<FString, FString>& InAdditionalSources,
+	FString const& InShaderSource,
+	TMap<FString, FString> const& InAdditionalSources,
 	TSharedPtr<FComputeKernelDefinitionSet>& InShaderDefinitionSet,
 	TSharedPtr<FComputeKernelPermutationVector>& InShaderPermutationVector,
 	TUniquePtr<FShaderParametersMetadataAllocations>& InShaderParameterMetadataAllocations,
@@ -230,8 +230,8 @@ void FComputeKernelResource::SetupResource(
 	FriendlyName = InFriendlyName;
 	ShaderEntryPoint = InShaderEntryPoint;
 	ShaderCodeHash = GetTypeHash(InShaderHashKey);
-	ShaderSource = MoveTemp(InShaderSource);
-	AdditionalSources = MoveTemp(InAdditionalSources);
+	ShaderSource = InShaderSource;
+	AdditionalSources = InAdditionalSources;
 	ShaderDefinitionSet = InShaderDefinitionSet;
 	ShaderPermutationVector = InShaderPermutationVector;
 	ShaderParameterMetadataAllocations = MoveTemp(InShaderParameterMetadataAllocations);
