ComputeFramework: Fix some platforms not compiling shaders due to missing platform specific includes.
#rnx

diff --git a/Engine/Plugins/Experimental/ComputeFramework/Shaders/Private/ComputeKernelCommon.ush b/Engine/Plugins/Experimental/ComputeFramework/Shaders/Private/ComputeKernelCommon.ush
--- a/Engine/Plugins/Experimental/ComputeFramework/Shaders/Private/ComputeKernelCommon.ush
+++ b/Engine/Plugins/Experimental/ComputeFramework/Shaders/Private/ComputeKernelCommon.ush
@@ -2,6 +2,8 @@
 
 #pragma once
 
+#include "/Engine/Public/Platform.ush"
+
 // Macro to mark a kernel entry point.
 // Only one should exist per kernel.
 
