ComputeGraph: Don't wait on IsCachedCookedPlatformDataLoaded() if no shader formats exist for a platform.
Fixes infinite cook of ComputeGraph objects on server platforms.
#rnx

diff --git a/Engine/Plugins/Runtime/ComputeFramework/Source/ComputeFramework/Private/ComputeGraph.cpp b/Engine/Plugins/Runtime/ComputeFramework/Source/ComputeFramework/Private/ComputeGraph.cpp
--- a/Engine/Plugins/Runtime/ComputeFramework/Source/ComputeFramework/Private/ComputeGraph.cpp
+++ b/Engine/Plugins/Runtime/ComputeFramework/Source/ComputeFramework/Private/ComputeGraph.cpp
@@ -744,6 +744,14 @@ void UComputeGraph::BeginCacheForCookedPlatformData(ITargetPlatform const* Targe
 
 bool UComputeGraph::IsCachedCookedPlatformDataLoaded(ITargetPlatform const* TargetPlatform)
 {
+	TArray<FName> ShaderFormats;
+	TargetPlatform->GetAllTargetedShaderFormats(ShaderFormats);
+	if (ShaderFormats.Num() == 0)
+	{
+		// Nothing will be queued in BeginCacheForCookedPlatformData.
+		return true;
+	}
+
 	for (int32 KernelIndex = 0; KernelIndex < KernelInvocations.Num(); ++KernelIndex)
 	{
 		UComputeKernelSource* KernelSource = KernelInvocations[KernelIndex] != nullptr ? KernelInvocations[KernelIndex]->KernelSource : nullptr;
