Optimus: added explict padding within structs so that structs with 1 float don't get collpased when used in a structured buffer

diff --git a/Engine/Plugins/Animation/DeformerGraph/Source/OptimusCore/Private/DataInterfaces/OptimusDataInterfaceAnimAttribute.cpp b/Engine/Plugins/Animation/DeformerGraph/Source/OptimusCore/Private/DataInterfaces/OptimusDataInterfaceAnimAttribute.cpp
--- a/Engine/Plugins/Animation/DeformerGraph/Source/OptimusCore/Private/DataInterfaces/OptimusDataInterfaceAnimAttribute.cpp
+++ b/Engine/Plugins/Animation/DeformerGraph/Source/OptimusCore/Private/DataInterfaces/OptimusDataInterfaceAnimAttribute.cpp
@@ -4,7 +4,6 @@
 
 #include "OptimusDataDomain.h"
 #include "OptimusDataTypeRegistry.h"
-#include "OptimusHelpers.h"
 #include "OptimusValueContainer.h"
 #include "RenderGraphBuilder.h"
 #include "RenderGraphResources.h"
@@ -14,6 +13,7 @@
 #include "Animation/BuiltInAttributeTypes.h"
 #include "Components/SkeletalMeshComponent.h"
 #include "ShaderParameterMetadataBuilder.h"
+#include "ComputeFramework/ComputeMetadataBuilder.h"
 #include "Engine/UserDefinedStruct.h"
 
 #include "Serialization/MemoryReader.h"
@@ -206,7 +206,7 @@ void UOptimusAnimAttributeDataInterface::GetShaderParameters(TCHAR const* UID, F
 	for (int32 Index = 0; Index < AttributeArray.Num(); Index++)
 	{
 		const FOptimusAnimAttributeDescription& Attribute = AttributeArray[Index];
-		Optimus::AddParamForType(Builder, *Attribute.HlslId, Attribute.DataType->ShaderValueType, NestedStructs);
+		ComputeFramework::AddParamForType(Builder, *Attribute.HlslId, Attribute.DataType->ShaderValueType, NestedStructs);
 	}
 
 	FShaderParametersMetadata* ShaderParameterMetadata = Builder.Build(FShaderParametersMetadata::EUseCase::ShaderParameterStruct, TEXT("UAnimAttributeDataInterface"));
@@ -574,7 +574,7 @@ void UOptimusAnimAttributeDataProvider::Init(
 	TArray<FShaderParametersMetadata*> NestedStructs;
 	for (FOptimusAnimAttributeDescription& Attribute : InAttributeArray)
 	{
-		Optimus::AddParamForType(Builder, *Attribute.Name, Attribute.DataType->ShaderValueType, NestedStructs);
+		ComputeFramework::AddParamForType(Builder, *Attribute.Name, Attribute.DataType->ShaderValueType, NestedStructs);
 	}
 
 	FShaderParametersMetadata* ShaderParameterMetadata = Builder.Build(FShaderParametersMetadata::EUseCase::ShaderParameterStruct, TEXT("UAnimAttributeDataInterface"));
diff --git a/Engine/Plugins/Animation/DeformerGraph/Source/OptimusCore/Private/DataInterfaces/OptimusDataInterfaceGraph.cpp b/Engine/Plugins/Animation/DeformerGraph/Source/OptimusCore/Private/DataInterfaces/OptimusDataInterfaceGraph.cpp
--- a/Engine/Plugins/Animation/DeformerGraph/Source/OptimusCore/Private/DataInterfaces/OptimusDataInterfaceGraph.cpp
+++ b/Engine/Plugins/Animation/DeformerGraph/Source/OptimusCore/Private/DataInterfaces/OptimusDataInterfaceGraph.cpp
@@ -5,10 +5,10 @@
 #include "Components/SkinnedMeshComponent.h"
 #include "ComputeFramework/ShaderParameterMetadataAllocation.h"
 #include "ComputeFramework/ShaderParamTypeDefinition.h"
-#include "OptimusHelpers.h"
 #include "OptimusDeformerInstance.h"
 #include "OptimusVariableDescription.h"
 #include "ShaderParameterMetadataBuilder.h"
+#include "ComputeFramework/ComputeMetadataBuilder.h"
 
 void UOptimusGraphDataInterface::Init(TArray<FOptimusGraphVariableDescription> const& InVariables)
 {
@@ -19,7 +19,7 @@ void UOptimusGraphDataInterface::Init(TArray<FOptimusGraphVariableDescription> c
 	TArray<FShaderParametersMetadata*> NestedStructMetadatas; 
 	for (FOptimusGraphVariableDescription const& Variable : Variables)
 	{
-		Optimus::AddParamForType(Builder, *Variable.Name, Variable.ValueType, NestedStructMetadatas);
+		ComputeFramework::AddParamForType(Builder, *Variable.Name, Variable.ValueType, NestedStructMetadatas);
 	}
 	
 	FShaderParametersMetadata* ShaderParameterMetadata = Builder.Build(FShaderParametersMetadata::EUseCase::ShaderParameterStruct, TEXT("UGraphDataInterface"));
@@ -60,7 +60,7 @@ void UOptimusGraphDataInterface::GetShaderParameters(TCHAR const* UID, FShaderPa
 	TArray<FShaderParametersMetadata*> NestedStructs;
 	for (FOptimusGraphVariableDescription const& Variable : Variables)
 	{
-		Optimus::AddParamForType(Builder, *Variable.Name, Variable.ValueType, NestedStructs);
+		ComputeFramework::AddParamForType(Builder, *Variable.Name, Variable.ValueType, NestedStructs);
 	}
 
 	FShaderParametersMetadata* ShaderParameterMetadata = Builder.Build(FShaderParametersMetadata::EUseCase::ShaderParameterStruct, TEXT("UGraphDataInterface"));
diff --git a/Engine/Plugins/Animation/DeformerGraph/Source/OptimusCore/Private/Nodes/OptimusNode_CustomComputeKernel.cpp b/Engine/Plugins/Animation/DeformerGraph/Source/OptimusCore/Private/Nodes/OptimusNode_CustomComputeKernel.cpp
--- a/Engine/Plugins/Animation/DeformerGraph/Source/OptimusCore/Private/Nodes/OptimusNode_CustomComputeKernel.cpp
+++ b/Engine/Plugins/Animation/DeformerGraph/Source/OptimusCore/Private/Nodes/OptimusNode_CustomComputeKernel.cpp
@@ -1044,7 +1044,7 @@ void UOptimusNode_CustomComputeKernel::UpdatePreamble()
 						const FString StructName = TypeHandle->ToString();
 						if (!StructsSeen.Contains(StructName))
 						{
-							Structs.Add(TypeHandle->GetTypeDeclaration(FriendlyNameMap) + TEXT(";\n\n"));
+							Structs.Add(TypeHandle->GetTypeDeclaration(FriendlyNameMap, true) + TEXT(";\n\n"));
 							StructsSeen.Add(StructName);
 						}	
 					}
diff --git a/Engine/Plugins/Animation/DeformerGraph/Source/OptimusCore/Private/OptimusDataTypeRegistry.cpp b/Engine/Plugins/Animation/DeformerGraph/Source/OptimusCore/Private/OptimusDataTypeRegistry.cpp
--- a/Engine/Plugins/Animation/DeformerGraph/Source/OptimusCore/Private/OptimusDataTypeRegistry.cpp
+++ b/Engine/Plugins/Animation/DeformerGraph/Source/OptimusCore/Private/OptimusDataTypeRegistry.cpp
@@ -11,6 +11,7 @@
 #include "Animation/BuiltInAttributeTypes.h"
 
 #include "AssetRegistry/AssetRegistryModule.h"
+#include "ComputeFramework/ComputeMetadataBuilder.h"
 
 #include "Engine/Classes/Engine/UserDefinedStruct.h"
 
@@ -609,7 +610,7 @@ bool FOptimusDataTypeRegistry::RegisterStructType(UScriptStruct* InStructType)
 		TArray<FArrayMetadata> ArrayMetadata;
 		
 
-		Optimus::FTypeMetaData TypeMetaData(ShaderValueType);
+		ComputeFramework::FTypeMetaData TypeMetaData(ShaderValueType);
 		ExpectedShaderValueSize = TypeMetaData.Metadata->GetSize();
 		int32 ExpectedPropertySize = 
 			InStructType->GetCppStructOps() ? 
diff --git a/Engine/Plugins/Animation/DeformerGraph/Source/OptimusCore/Private/OptimusHelpers.cpp b/Engine/Plugins/Animation/DeformerGraph/Source/OptimusCore/Private/OptimusHelpers.cpp
--- a/Engine/Plugins/Animation/DeformerGraph/Source/OptimusCore/Private/OptimusHelpers.cpp
+++ b/Engine/Plugins/Animation/DeformerGraph/Source/OptimusCore/Private/OptimusHelpers.cpp
@@ -42,61 +42,6 @@ FName Optimus::GetSanitizedNameForHlsl(FName InName)
 	return *Name;
 }
 
-template<typename T>
-void ParametrizedAddParm(FShaderParametersMetadataBuilder& InOutBuilder, const TCHAR* InName)
-{
-	InOutBuilder.AddParam<T>(InName);
-}
-
-void Optimus::AddParamForType(FShaderParametersMetadataBuilder& InOutBuilder, TCHAR const* InName, FShaderValueTypeHandle const& InValueType, TArray<FShaderParametersMetadata*>& OutNestedStructs)
-{
-	using AddParamFuncType = TFunction<void(FShaderParametersMetadataBuilder&, const TCHAR*)>;
-
-	static TMap<FShaderValueType, AddParamFuncType> AddParamFuncs =
-	{
-		{*FShaderValueType::Get(EShaderFundamentalType::Bool), &ParametrizedAddParm<bool>},
-		{*FShaderValueType::Get(EShaderFundamentalType::Int), &ParametrizedAddParm<int32>},
-		{*FShaderValueType::Get(EShaderFundamentalType::Int, 2), &ParametrizedAddParm<FIntPoint>},
-		{*FShaderValueType::Get(EShaderFundamentalType::Int, 3), &ParametrizedAddParm<FIntVector>},
-		{*FShaderValueType::Get(EShaderFundamentalType::Int, 4), &ParametrizedAddParm<FIntVector4>},
-		{*FShaderValueType::Get(EShaderFundamentalType::Uint), &ParametrizedAddParm<uint32>},
-		{*FShaderValueType::Get(EShaderFundamentalType::Uint, 2), &ParametrizedAddParm<FUintVector2>},
-		{*FShaderValueType::Get(EShaderFundamentalType::Uint, 4), &ParametrizedAddParm<FUintVector4>},
-		{*FShaderValueType::Get(EShaderFundamentalType::Float), &ParametrizedAddParm<float>},
-		{*FShaderValueType::Get(EShaderFundamentalType::Float, 2), &ParametrizedAddParm<FVector2f>},
-		{*FShaderValueType::Get(EShaderFundamentalType::Float, 3), &ParametrizedAddParm<FVector3f>},
-		{*FShaderValueType::Get(EShaderFundamentalType::Float, 4), &ParametrizedAddParm<FVector4f>},
-		{*FShaderValueType::Get(EShaderFundamentalType::Float, 4, 4), &ParametrizedAddParm<FMatrix44f>},
-	};
-
-	static TArray<FString> AllElementNames;
-	if (InValueType->bIsDynamicArray)
-	{
-		// both struct array and normal array are treated the same
-		InOutBuilder.AddRDGBufferSRV(InName, TEXT("StructuredBuffer"));
-	}
-	else if (InValueType->Type == EShaderFundamentalType::Struct)
-	{
-		FShaderParametersMetadataBuilder NestedStructBuilder;
-	
-		for (const FShaderValueType::FStructElement& Element : InValueType->StructElements)
-		{
-			AllElementNames.Add(Element.Name.ToString());
-			AddParamForType(NestedStructBuilder, *AllElementNames.Last(), Element.Type, OutNestedStructs);
-		}
-	
-		FShaderParametersMetadata* ShaderParameterMetadata = NestedStructBuilder.Build(FShaderParametersMetadata::EUseCase::ShaderParameterStruct, InName);
-	
-		InOutBuilder.AddNestedStruct(InName, ShaderParameterMetadata);
-	
-		OutNestedStructs.Add(ShaderParameterMetadata);
-	}
-	else if (const AddParamFuncType* Entry = AddParamFuncs.Find(*InValueType))
-	{
-		(*Entry)(InOutBuilder, InName);
-	}
-}
-
 bool Optimus::RenameObject(UObject* InObjectToRename, const TCHAR* InNewName, UObject* InNewOuter)
 {
 	return InObjectToRename->Rename(InNewName, InNewOuter, REN_ForceNoResetLoaders | REN_DoNotDirty | REN_DontCreateRedirectors | REN_NonTransactional);
@@ -119,25 +64,6 @@ TArray<UClass*> Optimus::GetClassObjectsInPackage(UPackage* InPackage)
 	return ClassObjects;
 }
 
-Optimus::FTypeMetaData::FTypeMetaData(FShaderValueTypeHandle InType)
-{
-	FShaderParametersMetadataBuilder Builder;
-	TArray<FShaderParametersMetadata*> NestedStructs;
-	Optimus::AddParamForType(Builder, TEXT("Dummy"), InType, NestedStructs);
-	FShaderParametersMetadata* ShaderParameterMetadata = Builder.Build(FShaderParametersMetadata::EUseCase::ShaderParameterStruct, TEXT("Dummy"));
-	Metadata = ShaderParameterMetadata->GetMembers()[0].GetStructMetadata();
-	AllocatedMetadatas.Add(ShaderParameterMetadata);
-	AllocatedMetadatas.Append(NestedStructs);
-}
-
-Optimus::FTypeMetaData::~FTypeMetaData()
-{
-	for (FShaderParametersMetadata* Allocated: AllocatedMetadatas)
-	{
-		delete Allocated;
-	}
-}
-
 FText Optimus::GetTypeDisplayName(UScriptStruct* InStructType)
 {
 #if WITH_EDITOR
diff --git a/Engine/Plugins/Animation/DeformerGraph/Source/OptimusCore/Public/OptimusHelpers.h b/Engine/Plugins/Animation/DeformerGraph/Source/OptimusCore/Public/OptimusHelpers.h
--- a/Engine/Plugins/Animation/DeformerGraph/Source/OptimusCore/Public/OptimusHelpers.h
+++ b/Engine/Plugins/Animation/DeformerGraph/Source/OptimusCore/Public/OptimusHelpers.h
@@ -72,8 +72,6 @@ namespace Optimus
 
 	FName GetSanitizedNameForHlsl(FName InName);
 
-	void AddParamForType(FShaderParametersMetadataBuilder& InOutBuilder, TCHAR const* InName, FShaderValueTypeHandle const& InValueType, TArray<FShaderParametersMetadata*>& OutNestedStructs);
-
 	FORCEINLINE_DEBUGGABLE FMatrix44f ConvertFTransformToFMatrix44f(const FTransform& InTransform)
 	{
 		return TransformConverter<FMatrix44f>::Convert<FMatrix44d>(InTransform.ToMatrixWithScale());	
@@ -85,19 +83,6 @@ namespace Optimus
 		to collect them */
 	TArray<UClass*> GetClassObjectsInPackage(UPackage* InPackage);
 
-	/** Helper struct to convert shader value type to shader parameter metadata */
-	struct FTypeMetaData
-	{
-		FTypeMetaData(FShaderValueTypeHandle InType);
-		FTypeMetaData(const FTypeMetaData& InOther) = delete;
-		FTypeMetaData& operator=(const FTypeMetaData& InOther) = delete;
-		~FTypeMetaData();
-
-		const FShaderParametersMetadata* Metadata;
-		
-		TArray<FShaderParametersMetadata*> AllocatedMetadatas;
-	};
-	
 	/** Return the unique type name for registry and kernel code generation if the bInShouldGetUniqueNameForUserDefinedStruct = true
 		Otherise, it returns the friendly name for user-facing shader text display*/
 	OPTIMUSCORE_API FName GetTypeName(UScriptStruct* InStruct, bool bInShouldGetUniqueNameForUserDefinedStruct = true);
diff --git a/Engine/Plugins/Animation/DeformerGraph/Source/OptimusEditor/Private/OptimusDetailsCustomization.cpp b/Engine/Plugins/Animation/DeformerGraph/Source/OptimusEditor/Private/OptimusDetailsCustomization.cpp
--- a/Engine/Plugins/Animation/DeformerGraph/Source/OptimusEditor/Private/OptimusDetailsCustomization.cpp
+++ b/Engine/Plugins/Animation/DeformerGraph/Source/OptimusEditor/Private/OptimusDetailsCustomization.cpp
@@ -16,6 +16,7 @@
 #include "OptimusDeformerInstance.h"
 #include "OptimusEditorStyle.h"
 #include "OptimusExpressionEvaluator.h"
+#include "OptimusHelpers.h"
 #include "OptimusHLSLSyntaxHighlighter.h"
 #include "OptimusNode.h"
 #include "OptimusResourceDescription.h"
@@ -26,6 +27,7 @@
 #include "PropertyEditor/Private/PropertyNode.h"
 #include "PropertyEditor/Public/IPropertyUtilities.h"
 #include "ScopedTransaction.h"
+#include "Engine/UserDefinedStruct.h"
 #include "Styling/AppStyle.h"
 #include "Styling/SlateIconFinder.h"
 #include "Widgets/Input/SButton.h"
@@ -192,11 +194,25 @@ FText FOptimusDataTypeRefCustomization::GetDeclarationText() const
 
 	if (DataType.IsValid() && DataType->ShaderValueType.IsValid())
 	{
-		const FShaderValueType* ValueType = DataType->ShaderValueType.ValueTypePtr;
+		const FShaderValueTypeHandle& ValueType = DataType->ShaderValueType;
 		FText Declaration;
 		if (ValueType->Type == EShaderFundamentalType::Struct)
 		{
-			return FText::FromString(ValueType->GetTypeDeclaration());
+			TArray<FShaderValueTypeHandle> StructTypes = ValueType->GetMemberStructTypes();
+			StructTypes.Add(ValueType);
+
+			// Collect all friendly names
+			TMap<FName, FName> FriendlyNameMap;
+			for (const FShaderValueTypeHandle& TypeHandle : StructTypes )
+			{
+				const FOptimusDataTypeHandle DataTypeHandle = FOptimusDataTypeRegistry::Get().FindType(TypeHandle->Name);
+				if (UUserDefinedStruct* Struct = Cast<UUserDefinedStruct>(DataTypeHandle->TypeObject))
+				{
+					FriendlyNameMap.Add(TypeHandle->Name) = Optimus::GetTypeName(Struct, false);
+				}
+			}
+	
+			return FText::FromString(ValueType->GetTypeDeclaration(FriendlyNameMap, true));
 		}
 		else
 		{
diff --git a/Engine/Plugins/Runtime/ComputeFramework/Source/ComputeFramework/Private/ComputeMetadataBuilder.cpp b/Engine/Plugins/Runtime/ComputeFramework/Source/ComputeFramework/Private/ComputeMetadataBuilder.cpp
new file mode 100644
--- /dev/null
+++ b/Engine/Plugins/Runtime/ComputeFramework/Source/ComputeFramework/Private/ComputeMetadataBuilder.cpp
@@ -0,0 +1,82 @@
+// Copyright Epic Games, Inc. All Rights Reserved.
+
+#pragma once
+
+#include "ComputeFramework/Public/ComputeFramework/ComputeMetadataBuilder.h"
+
+#include "ShaderParameterMetadataBuilder.h"
+#include "ComputeFramework/ShaderParamTypeDefinition.h"
+
+template<typename T>
+void ParametrizedAddParm(FShaderParametersMetadataBuilder& InOutBuilder, const TCHAR* InName)
+{
+	InOutBuilder.AddParam<T>(InName);
+}
+
+void ComputeFramework::AddParamForType(FShaderParametersMetadataBuilder& InOutBuilder, TCHAR const* InName, FShaderValueTypeHandle const& InValueType, TArray<FShaderParametersMetadata*>& OutNestedStructs)
+{
+	using AddParamFuncType = TFunction<void(FShaderParametersMetadataBuilder&, const TCHAR*)>;
+
+	static TMap<FShaderValueType, AddParamFuncType> AddParamFuncs =
+	{
+		{*FShaderValueType::Get(EShaderFundamentalType::Bool), &ParametrizedAddParm<bool>},
+		{*FShaderValueType::Get(EShaderFundamentalType::Int), &ParametrizedAddParm<int32>},
+		{*FShaderValueType::Get(EShaderFundamentalType::Int, 2), &ParametrizedAddParm<FIntPoint>},
+		{*FShaderValueType::Get(EShaderFundamentalType::Int, 3), &ParametrizedAddParm<FIntVector>},
+		{*FShaderValueType::Get(EShaderFundamentalType::Int, 4), &ParametrizedAddParm<FIntVector4>},
+		{*FShaderValueType::Get(EShaderFundamentalType::Uint), &ParametrizedAddParm<uint32>},
+		{*FShaderValueType::Get(EShaderFundamentalType::Uint, 2), &ParametrizedAddParm<FUintVector2>},
+		{*FShaderValueType::Get(EShaderFundamentalType::Uint, 4), &ParametrizedAddParm<FUintVector4>},
+		{*FShaderValueType::Get(EShaderFundamentalType::Float), &ParametrizedAddParm<float>},
+		{*FShaderValueType::Get(EShaderFundamentalType::Float, 2), &ParametrizedAddParm<FVector2f>},
+		{*FShaderValueType::Get(EShaderFundamentalType::Float, 3), &ParametrizedAddParm<FVector3f>},
+		{*FShaderValueType::Get(EShaderFundamentalType::Float, 4), &ParametrizedAddParm<FVector4f>},
+		{*FShaderValueType::Get(EShaderFundamentalType::Float, 4, 4), &ParametrizedAddParm<FMatrix44f>},
+	};
+
+	static TArray<FString> AllElementNames;
+	if (InValueType->bIsDynamicArray)
+	{
+		// both struct array and normal array are treated the same
+		InOutBuilder.AddRDGBufferSRV(InName, TEXT("StructuredBuffer"));
+	}
+	else if (InValueType->Type == EShaderFundamentalType::Struct)
+	{
+		FShaderParametersMetadataBuilder NestedStructBuilder;
+	
+		for (const FShaderValueType::FStructElement& Element : InValueType->StructElements)
+		{
+			AllElementNames.Add(Element.Name.ToString());
+			AddParamForType(NestedStructBuilder, *AllElementNames.Last(), Element.Type, OutNestedStructs);
+		}
+	
+		FShaderParametersMetadata* ShaderParameterMetadata = NestedStructBuilder.Build(FShaderParametersMetadata::EUseCase::ShaderParameterStruct, InName);
+	
+		InOutBuilder.AddNestedStruct(InName, ShaderParameterMetadata);
+	
+		OutNestedStructs.Add(ShaderParameterMetadata);
+	}
+	else if (const AddParamFuncType* Entry = AddParamFuncs.Find(*InValueType))
+	{
+		(*Entry)(InOutBuilder, InName);
+	}
+}
+
+ComputeFramework::FTypeMetaData::FTypeMetaData(FShaderValueTypeHandle InType)
+{
+	FShaderParametersMetadataBuilder Builder;
+	TArray<FShaderParametersMetadata*> NestedStructs;
+	AddParamForType(Builder, TEXT("Dummy"), InType, NestedStructs);
+	FShaderParametersMetadata* ShaderParameterMetadata = Builder.Build(FShaderParametersMetadata::EUseCase::ShaderParameterStruct, TEXT("Dummy"));
+	Metadata = ShaderParameterMetadata->GetMembers()[0].GetStructMetadata();
+	AllocatedMetadatas.Add(ShaderParameterMetadata);
+	AllocatedMetadatas.Append(NestedStructs);
+}
+
+ComputeFramework::FTypeMetaData::~FTypeMetaData()
+{
+	for (FShaderParametersMetadata* Allocated: AllocatedMetadatas)
+	{
+		delete Allocated;
+	}
+}
diff --git a/Engine/Plugins/Runtime/ComputeFramework/Source/ComputeFramework/Private/ShaderParamTypeDefinition.cpp b/Engine/Plugins/Runtime/ComputeFramework/Source/ComputeFramework/Private/ShaderParamTypeDefinition.cpp
--- a/Engine/Plugins/Runtime/ComputeFramework/Source/ComputeFramework/Private/ShaderParamTypeDefinition.cpp
+++ b/Engine/Plugins/Runtime/ComputeFramework/Source/ComputeFramework/Private/ShaderParamTypeDefinition.cpp
@@ -2,7 +2,9 @@
 
 #include "ComputeFramework/ShaderParamTypeDefinition.h"
 
+#include "ShaderParameterMetadataBuilder.h"
 #include "Algo/Find.h"
+#include "ComputeFramework/ComputeMetadataBuilder.h"
 #include "ComputeFramework/ComputeFrameworkObjectVersion.h"
 #include "Internationalization/Regex.h"
 #include "Misc/DefaultValueHelper.h"
@@ -377,7 +379,7 @@ FString FShaderValueType::ToString(const FName& InStructTypeNameOverride) const
 	return BaseName;
 }
 
-FString FShaderValueType::GetTypeDeclaration(const TMap<FName, FName>& InTypeNamesToReplace) const
+FString FShaderValueType::GetTypeDeclaration(const TMap<FName, FName>& InTypeNamesToReplace, bool bCommentPaddings) const
 {
 	// FIXME: Cache on create?
 	if (Type != EShaderFundamentalType::Struct)
@@ -385,10 +387,83 @@ FString FShaderValueType::GetTypeDeclaration(const TMap<FName, FName>& InTypeNam
 		return {};
 	}
 
+	FShaderValueTypeHandle TypeHandle;
+	TypeHandle.ValueTypePtr = this;
+	
+	if (bIsDynamicArray)
+	{
+		FShaderValueType ElementType = *this;
+		ElementType.bIsDynamicArray = false;
+		TypeHandle = GetOrCreate(MoveTemp(ElementType));
+	}
+	
+	ComputeFramework::FTypeMetaData TypeMetaData(TypeHandle);
+	const TArray<FShaderParametersMetadata::FMember>& MemberDatas = TypeMetaData.Metadata->GetMembers();
+	
+	if (!ensure(MemberDatas.Num() == StructElements.Num()))
+	{
+		return {};
+	}
+
+	// Record where each padding starts	
+	TArray<int32> PadOffsets;
+	PadOffsets.Reserve(MemberDatas.Num());
+	
+	// Total allocated size includes member size + per member padding
+	TArray<int32> MemberTotalSizes;
+	MemberTotalSizes.Reserve(MemberDatas.Num());
+	
+	for (int32 Index = 0; Index < MemberDatas.Num(); Index++)
+	{
+		if (Index + 1 < MemberDatas.Num())
+		{
+			MemberTotalSizes.Add(MemberDatas[Index + 1].GetOffset() - MemberDatas[Index].GetOffset());
+		}
+		else
+		{
+			MemberTotalSizes.Add(TypeMetaData.Metadata->GetSize() - MemberDatas[Index].GetOffset());
+		}
+
+		// Initialize the pad offsets to be the member offset
+		PadOffsets.Add(MemberDatas[Index].GetOffset());
+	}
+
+	// Calculate the padding size, should be in multiples of 4
+	TArray<int32> Pads;
+	Pads.Reserve(MemberDatas.Num());
+	for (int32 Index = 0; Index < StructElements.Num(); Index++)
+	{
+		const FStructElement& StructElement = StructElements[Index];
+		int32 InlineSize = 0;
+		if (StructElement.Type->bIsDynamicArray)
+		{
+			// Array pointers are always of size SHADER_PARAMETER_POINTER_ALIGNMENT
+			InlineSize = SHADER_PARAMETER_POINTER_ALIGNMENT;
+		}
+		else if (StructElement.Type->StructElements.Num() > 0)
+		{
+			// Struct type does not need external padding
+			// it should be padded internally to align to 16 byte boundary 
+			InlineSize = MemberTotalSizes[Index];
+		}
+		else
+		{
+			InlineSize = StructElement.Type->GetResourceElementSize();
+		}
+
+		// Padding starting position = Member starting position + member size;
+		PadOffsets[Index] += InlineSize;
+		
+		int32 Pad = MemberTotalSizes[Index] - InlineSize;
+		Pads.Add(Pad);
+	}
+
 	TArray<FString> Elements;
 	Elements.Reserve(StructElements.Num());
-	for (const FStructElement& StructElement : StructElements)
+	for (int32 Index = 0; Index < StructElements.Num(); Index++)
 	{
+		const FStructElement& StructElement = StructElements[Index];
+		
 		FName MemberTypeNameToUse = StructElement.Type->Name;
 		if (const FName* Replacement = InTypeNamesToReplace.Find(StructElement.Type->Name))
 		{
@@ -397,10 +472,24 @@ FString FShaderValueType::GetTypeDeclaration(const TMap<FName, FName>& InTypeNam
 		
 		Elements.Add(FString::Printf(TEXT("    %s %s;\n"), *StructElement.Type.ValueTypePtr->ToString(MemberTypeNameToUse), *StructElement.Name.ToString()));
 
-		if (StructElement.Type->bIsDynamicArray)
+		// Add padding elements
+		if (Pads[Index] > 0)
 		{
-			FString ElementType = StructElement.Type.ValueTypePtr->ToString();
-			ElementType = ElementType.Replace(TEXT("StructuredBuffer"),TEXT("")).Replace(TEXT("<"), TEXT("")).Replace(TEXT(">"), TEXT(""));
+			if (ensure(Pads[Index] % 4 == 0))
+			{
+				const int32 NumPads = Pads[Index] / 4;
+
+				const FString CommentPrefix = bCommentPaddings ? TEXT("// ") : TEXT("");
+				
+				if (NumPads == 1)
+				{
+					Elements.Add(FString::Printf(TEXT("    %sfloat Padding_%s_%d;\n"), *CommentPrefix, *StructElement.Name.ToString(), PadOffsets[Index]));
+				}
+				else
+				{
+					Elements.Add(FString::Printf(TEXT("    %sfloat Padding_%s_%d[%d];\n"), *CommentPrefix, *StructElement.Name.ToString(), PadOffsets[Index], NumPads));
+				}
+			}
 		}
 	}
 
diff --git a/Engine/Plugins/Runtime/ComputeFramework/Source/ComputeFramework/Public/ComputeFramework/ComputeMetadataBuilder.h b/Engine/Plugins/Runtime/ComputeFramework/Source/ComputeFramework/Public/ComputeFramework/ComputeMetadataBuilder.h
new file mode 100644
--- /dev/null
+++ b/Engine/Plugins/Runtime/ComputeFramework/Source/ComputeFramework/Public/ComputeFramework/ComputeMetadataBuilder.h
@@ -0,0 +1,27 @@
+// Copyright Epic Games, Inc. All Rights Reserved.
+
+#pragma once
+
+#include "CoreMinimal.h"
+
+struct FShaderValueTypeHandle;
+class FShaderParametersMetadataBuilder;
+class FShaderParametersMetadata;
+
+namespace ComputeFramework
+{
+	COMPUTEFRAMEWORK_API void AddParamForType(FShaderParametersMetadataBuilder& InOutBuilder, TCHAR const* InName, FShaderValueTypeHandle const& InValueType, TArray<FShaderParametersMetadata*>& OutNestedStructs);
+
+	/** Helper struct to convert shader value type to shader parameter metadata */
+	 struct COMPUTEFRAMEWORK_API FTypeMetaData
+	{
+		FTypeMetaData(FShaderValueTypeHandle InType);
+		FTypeMetaData(const FTypeMetaData& InOther) = delete;
+		FTypeMetaData& operator=(const FTypeMetaData& InOther) = delete;
+		~FTypeMetaData();
+
+		const FShaderParametersMetadata* Metadata;
+		
+		TArray<FShaderParametersMetadata*> AllocatedMetadatas;
+	};	
+}
diff --git a/Engine/Plugins/Runtime/ComputeFramework/Source/ComputeFramework/Public/ComputeFramework/ShaderParamTypeDefinition.h b/Engine/Plugins/Runtime/ComputeFramework/Source/ComputeFramework/Public/ComputeFramework/ShaderParamTypeDefinition.h
--- a/Engine/Plugins/Runtime/ComputeFramework/Source/ComputeFramework/Public/ComputeFramework/ShaderParamTypeDefinition.h
+++ b/Engine/Plugins/Runtime/ComputeFramework/Source/ComputeFramework/Public/ComputeFramework/ShaderParamTypeDefinition.h
@@ -209,7 +209,7 @@ struct COMPUTEFRAMEWORK_API FShaderValueType
 	FString ToString(const FName& InStructTypeNameOverride = {}) const;
 
 	/** Returns the type declaration if this type is a struct, or the empty string if not. */
-	FString GetTypeDeclaration(const TMap<FName, FName>& InNamesToReplace = {}) const;
+	FString GetTypeDeclaration(const TMap<FName, FName>& InNamesToReplace = {}, bool bCommentPaddings = false) const;
 
 	/** Returns the all the struct types used in this type if the type is a struct. */
 	TArray<FShaderValueTypeHandle> GetMemberStructTypes() const;
