ComputeFramework, fix non-unity

#horde https://horde.devtools.epicgames.com/job/63296512f855485d69953910?step=6f3a

diff --git a/Engine/Plugins/Runtime/ComputeFramework/Source/ComputeFramework/Private/ComputeMetadataBuilder.cpp b/Engine/Plugins/Runtime/ComputeFramework/Source/ComputeFramework/Private/ComputeMetadataBuilder.cpp
--- a/Engine/Plugins/Runtime/ComputeFramework/Source/ComputeFramework/Private/ComputeMetadataBuilder.cpp
+++ b/Engine/Plugins/Runtime/ComputeFramework/Source/ComputeFramework/Private/ComputeMetadataBuilder.cpp
@@ -1,7 +1,5 @@
 // Copyright Epic Games, Inc. All Rights Reserved.
 
-#pragma once
-
 #include "ComputeFramework/Public/ComputeFramework/ComputeMetadataBuilder.h"
 
 #include "ShaderParameterMetadataBuilder.h"
