Fixed uniform buffer leak in Niagara due to parameters being allocated without destructor tracking.

diff --git a/Engine/Plugins/FX/Niagara/Source/Niagara/Private/NiagaraGpuComputeDispatch.cpp b/Engine/Plugins/FX/Niagara/Source/Niagara/Private/NiagaraGpuComputeDispatch.cpp
--- a/Engine/Plugins/FX/Niagara/Source/Niagara/Private/NiagaraGpuComputeDispatch.cpp
+++ b/Engine/Plugins/FX/Niagara/Source/Niagara/Private/NiagaraGpuComputeDispatch.cpp
@@ -1434,8 +1434,7 @@ void FNiagaraGpuComputeDispatch::DispatchStage(FRDGBuilder& GraphBuilder, const
 
 	const FNiagaraShaderScriptParametersMetadata& NiagaraShaderParametersMetadata = InstanceData.Context->GPUScript_RT->GetScriptParametersMetadata().Get();
 	const FShaderParametersMetadata* ShaderParametersMetadata = NiagaraShaderParametersMetadata.ShaderParametersMetadata.Get();
-	FNiagaraShader::FParameters* DispatchParameters = reinterpret_cast<FNiagaraShader::FParameters*>(GraphBuilder.Alloc(ShaderParametersMetadata->GetSize(), SHADER_PARAMETER_STRUCT_ALIGNMENT));
-	FMemory::Memset(DispatchParameters, 0, ShaderParametersMetadata->GetSize());
+	FNiagaraShader::FParameters* DispatchParameters = GraphBuilder.AllocParameters<FNiagaraShader::FParameters>(ShaderParametersMetadata);
 
 	// Set Parameters
 	const bool bRequiresPersistentIDs = InstanceData.Context->MainDataSet->RequiresPersistentIDs();
diff --git a/Engine/Plugins/Runtime/ComputeFramework/Source/ComputeFramework/Private/ComputeGraphWorker.cpp b/Engine/Plugins/Runtime/ComputeFramework/Source/ComputeFramework/Private/ComputeGraphWorker.cpp
--- a/Engine/Plugins/Runtime/ComputeFramework/Source/ComputeFramework/Private/ComputeGraphWorker.cpp
+++ b/Engine/Plugins/Runtime/ComputeFramework/Source/ComputeFramework/Private/ComputeGraphWorker.cpp
@@ -62,11 +62,7 @@ void FComputeGraphTaskWorker::SubmitWork(FRHICommandListImmediate& RHICmdList, E
 				TArray<FIntVector> ThreadCounts;
 				const int32 NumSubInvocations = GraphInvocation.DataProviderRenderProxies[KernelInvocation.ExecutionProviderIndex]->GetDispatchThreadCount(ThreadCounts);
 
-				const int32 ParameterBufferSize = Align(KernelInvocation.ShaderParameterMetadata->GetSize(), SHADER_PARAMETER_STRUCT_ALIGNMENT);
-
-				void* RawBuffer = GraphBuilder.Alloc(NumSubInvocations * ParameterBufferSize, SHADER_PARAMETER_STRUCT_ALIGNMENT);
-				FMemory::Memzero(RawBuffer, NumSubInvocations * ParameterBufferSize);
-				uint8* ParameterBuffer = static_cast<uint8*>(RawBuffer);
+				TStridedView<FComputeKernelShader::FParameters> ParameterArray = GraphBuilder.AllocParameters<FComputeKernelShader::FParameters>(KernelInvocation.ShaderParameterMetadata, NumSubInvocations);
 
 				// Iterate shader parameter members to fill the dispatch data structures.
 				// We assume that the members were filled out with a single data interface per member, and that the
@@ -74,10 +70,10 @@ void FComputeGraphTaskWorker::SubmitWork(FRHICommandListImmediate& RHICmdList, E
 				TArray<FShaderParametersMetadata::FMember> const& ParamMembers = KernelInvocation.ShaderParameterMetadata->GetMembers();
 
 				FComputeDataProviderRenderProxy::FCollectedDispatchData DispatchData;
-				DispatchData.ParameterBuffer = ParameterBuffer;
+				DispatchData.ParameterBuffer = reinterpret_cast<uint8*>(&ParameterArray[0]);
 				DispatchData.PermutationId.AddZeroed(NumSubInvocations);
 
-				FComputeDataProviderRenderProxy::FDispatchSetup DispatchSetup{ NumSubInvocations, 0, ParameterBufferSize, 0, GraphRenderProxy->ShaderPermutationVectors[KernelIndex]};
+				FComputeDataProviderRenderProxy::FDispatchSetup DispatchSetup{ NumSubInvocations, 0, ParameterArray.GetStride(), 0, GraphRenderProxy->ShaderPermutationVectors[KernelIndex]};
 
 				for (int32 MemberIndex = 0; MemberIndex < ParamMembers.Num(); ++MemberIndex)
 				{
@@ -108,7 +104,7 @@ void FComputeGraphTaskWorker::SubmitWork(FRHICommandListImmediate& RHICmdList, E
 						ERDGPassFlags::Compute | ERDGPassFlags::NeverCull,
 						Shader,
 						KernelInvocation.ShaderParameterMetadata,
-						reinterpret_cast<FComputeKernelShader::FParameters*>(ParameterBuffer + ParameterBufferSize * SubInvocationIndex),
+						&ParameterArray[SubInvocationIndex],
 						GroupCount
 					);
 				}
diff --git a/Engine/Source/Runtime/RenderCore/Public/RenderGraphBuilder.h b/Engine/Source/Runtime/RenderCore/Public/RenderGraphBuilder.h
--- a/Engine/Source/Runtime/RenderCore/Public/RenderGraphBuilder.h
+++ b/Engine/Source/Runtime/RenderCore/Public/RenderGraphBuilder.h
@@ -9,6 +9,7 @@
 #include "Containers/Set.h"
 #include "Containers/SortedMap.h"
 #include "Containers/UnrealString.h"
+#include "Containers/StridedView.h"
 #include "CoreTypes.h"
 #include "Misc/AssertionMacros.h"
 #include "MultiGPU.h"
@@ -158,6 +159,14 @@ public:
 	template <typename ParameterStructType>
 	ParameterStructType* AllocParameters(ParameterStructType* StructToCopy);
 
+	/** Allocates a data-driven parameter struct with a lifetime tied to graph execution. */
+	template <typename BaseParameterStructType>
+	BaseParameterStructType* AllocParameters(const FShaderParametersMetadata* ParametersMetadata);
+
+	/** Allocates an array of data-driven parameter structs with a lifetime tied to graph execution. */
+	template <typename BaseParameterStructType>
+	TStridedView<BaseParameterStructType> AllocParameters(const FShaderParametersMetadata* ParametersMetadata, uint32 NumStructs);
+
 	//////////////////////////////////////////////////////////////////////////
 
 	/** Adds a lambda pass to the graph with an accompanied pass parameter struct.
diff --git a/Engine/Source/Runtime/RenderCore/Public/RenderGraphBuilder.inl b/Engine/Source/Runtime/RenderCore/Public/RenderGraphBuilder.inl
--- a/Engine/Source/Runtime/RenderCore/Public/RenderGraphBuilder.inl
+++ b/Engine/Source/Runtime/RenderCore/Public/RenderGraphBuilder.inl
@@ -169,6 +169,46 @@ FORCEINLINE ParameterStructType* FRDGBuilder::AllocParameters(ParameterStructTyp
 	return Struct;
 }
 
+template <typename BaseParameterStructType>
+BaseParameterStructType* FRDGBuilder::AllocParameters(const FShaderParametersMetadata* ParametersMetadata)
+{
+	return &AllocParameters<BaseParameterStructType>(ParametersMetadata, 1)[0];
+}
+
+template <typename BaseParameterStructType>
+TStridedView<BaseParameterStructType> FRDGBuilder::AllocParameters(const FShaderParametersMetadata* ParametersMetadata, uint32 NumStructs)
+{
+	// NOTE: Contents are always zero! This might differ if shader parameters have a non-zero default initializer.
+	const int32 Stride = ParametersMetadata->GetSize();
+	BaseParameterStructType* Contents = reinterpret_cast<BaseParameterStructType*>(Allocator.Alloc(Stride * NumStructs, SHADER_PARAMETER_STRUCT_ALIGNMENT));
+	FMemory::Memset(Contents, 0, Stride * NumStructs);
+	TStridedView<BaseParameterStructType> ParameterArray(Stride, Contents, NumStructs);
+
+	struct FClearUniformBuffers
+	{
+	public:
+		FClearUniformBuffers(TStridedView<BaseParameterStructType> InParameterArray, const FRHIUniformBufferLayout& InLayout)
+			: ParameterArray(InParameterArray)
+			, Layout(&InLayout)
+		{}
+
+		~FClearUniformBuffers()
+		{
+			for (BaseParameterStructType& ParameterStruct : ParameterArray)
+			{
+				FRDGParameterStruct::ClearUniformBuffers(&ParameterStruct, Layout);
+			}
+		}
+
+	private:
+		TStridedView<BaseParameterStructType> ParameterArray;
+		const FRHIUniformBufferLayout* Layout;
+	};
+
+	AllocObject<FClearUniformBuffers>(ParameterArray, ParametersMetadata->GetLayout());
+	return ParameterArray;
+}
+
 FORCEINLINE FRDGSubresourceState* FRDGBuilder::AllocSubresource(const FRDGSubresourceState& Other)
 {
 	return Allocator.AllocNoDestruct<FRDGSubresourceState>(Other);
diff --git a/Engine/Source/Runtime/RenderCore/Public/RenderGraphParameter.h b/Engine/Source/Runtime/RenderCore/Public/RenderGraphParameter.h
--- a/Engine/Source/Runtime/RenderCore/Public/RenderGraphParameter.h
+++ b/Engine/Source/Runtime/RenderCore/Public/RenderGraphParameter.h
@@ -295,6 +295,9 @@ public:
 	/** Returns the render pass info generated from the render target binding slots. */
 	FRHIRenderPassInfo GetRenderPassInfo() const;
 
+	/** Clears out all uniform buffer references in the parameter struct. */
+	static void ClearUniformBuffers(void* Contents, const FRHIUniformBufferLayout* Layout);
+
 private:
 	FRDGParameter GetParameterInternal(TArrayView<const FRHIUniformBufferResource> Parameters, uint32 ParameterIndex) const
 	{
diff --git a/Engine/Source/Runtime/RenderCore/Public/RenderGraphParameters.inl b/Engine/Source/Runtime/RenderCore/Public/RenderGraphParameters.inl
--- a/Engine/Source/Runtime/RenderCore/Public/RenderGraphParameters.inl
+++ b/Engine/Source/Runtime/RenderCore/Public/RenderGraphParameters.inl
@@ -56,4 +56,19 @@ void FRDGParameterStruct::EnumerateUniformBuffers(FunctionType Function) const
 			Function(UniformBuffer);
 		}
 	}
+}
+
+inline void FRDGParameterStruct::ClearUniformBuffers(void* Contents, const FRHIUniformBufferLayout* Layout)
+{
+	FRDGParameterStruct Parameters(Contents, Layout);
+
+	for (uint32 Index = 0, Count = Layout->UniformBuffers.Num(); Index < Count; ++Index)
+	{
+		const FRDGParameter Parameter = Parameters.GetParameterInternal(Layout->UniformBuffers, Index);
+
+		if (FUniformBufferBinding* UniformBuffer = Parameter.GetAs<FUniformBufferBinding>())
+		{
+			*UniformBuffer = {};
+		}
+	}
 }
\ No newline at end of file
