Optimus: Added some defensive code where people have seen issues.
#rnx

diff --git a/Engine/Plugins/Animation/DeformerGraph/Source/OptimusCore/Private/OptimusNodePin.cpp b/Engine/Plugins/Animation/DeformerGraph/Source/OptimusCore/Private/OptimusNodePin.cpp
--- a/Engine/Plugins/Animation/DeformerGraph/Source/OptimusCore/Private/OptimusNodePin.cpp
+++ b/Engine/Plugins/Animation/DeformerGraph/Source/OptimusCore/Private/OptimusNodePin.cpp
@@ -419,7 +419,12 @@ TArray<UOptimusNodePin*> UOptimusNodePin::GetSubPinsRecursively(
 
 TArray<UOptimusNodePin*> UOptimusNodePin::GetConnectedPins() const
 {
-	return GetOwningNode()->GetOwningGraph()->GetConnectedPins(this);
+	UOptimusNodeGraph* Graph = GetOwningNode()->GetOwningGraph();
+	if (ensure(Graph != nullptr))
+	{
+		return Graph->GetConnectedPins(this);
+	}
+	return {};
 }
 
 
diff --git a/Engine/Plugins/Runtime/ComputeFramework/Source/ComputeFramework/Private/ComputeGraph.cpp b/Engine/Plugins/Runtime/ComputeFramework/Source/ComputeFramework/Private/ComputeGraph.cpp
--- a/Engine/Plugins/Runtime/ComputeFramework/Source/ComputeFramework/Private/ComputeGraph.cpp
+++ b/Engine/Plugins/Runtime/ComputeFramework/Source/ComputeFramework/Private/ComputeGraph.cpp
@@ -338,9 +338,13 @@ FComputeGraphRenderProxy* UComputeGraph::CreateRenderProxy() const
 				{
 					Invocation.BoundProviderIndices.AddUnique(GraphEdge.DataInterfaceIndex);
 
-					if (DataInterfaces[GraphEdge.DataInterfaceIndex]->IsExecutionInterface())
+					UComputeDataInterface const* DataInterface = DataInterfaces[GraphEdge.DataInterfaceIndex];
+					if (ensure(DataInterface != nullptr))
 					{
-						Invocation.ExecutionProviderIndex = GraphEdge.DataInterfaceIndex;
+						if (DataInterface->IsExecutionInterface())
+						{
+							Invocation.ExecutionProviderIndex = GraphEdge.DataInterfaceIndex;
+						}
 					}
 				}
 			}
