Fix use-after-free in UComputeGraph due to passing a TCHAR* from a temp FString to FShaderParameterMetadataBuilder

diff --git a/Engine/Plugins/Runtime/ComputeFramework/Source/ComputeFramework/Private/ComputeGraph.cpp b/Engine/Plugins/Runtime/ComputeFramework/Source/ComputeFramework/Private/ComputeGraph.cpp
--- a/Engine/Plugins/Runtime/ComputeFramework/Source/ComputeFramework/Private/ComputeGraph.cpp
+++ b/Engine/Plugins/Runtime/ComputeFramework/Source/ComputeFramework/Private/ComputeGraph.cpp
@@ -278,7 +278,10 @@ FShaderParametersMetadata* UComputeGraph::BuildKernelShaderMetadata(int32 InKern
 		}
 	}
 
-	FShaderParametersMetadata* ShaderParameterMetadata = Builder.Build(FShaderParametersMetadata::EUseCase::ShaderParameterStruct, *GetName());
+	// Graph name needs to persist since it's referenced by the metadata
+	int32 GraphNameIndex = InOutAllocations.Names.Add();
+	InOutAllocations.Names[GraphNameIndex] = GetName();
+	FShaderParametersMetadata* ShaderParameterMetadata = Builder.Build(FShaderParametersMetadata::EUseCase::ShaderParameterStruct, *InOutAllocations.Names[GraphNameIndex]);
 	InOutAllocations.ShaderParameterMetadatas.Add(ShaderParameterMetadata);
 
 	return ShaderParameterMetadata;
