Reset compiler output before recompiling to dump shader code

diff --git a/Engine/Plugins/Runtime/ComputeFramework/Source/ComputeFramework/Private/ComputeKernelShaderCompilationManager.cpp b/Engine/Plugins/Runtime/ComputeFramework/Source/ComputeFramework/Private/ComputeKernelShaderCompilationManager.cpp
--- a/Engine/Plugins/Runtime/ComputeFramework/Source/ComputeFramework/Private/ComputeKernelShaderCompilationManager.cpp
+++ b/Engine/Plugins/Runtime/ComputeFramework/Source/ComputeFramework/Private/ComputeKernelShaderCompilationManager.cpp
@@ -140,6 +140,7 @@ void FComputeKernelShaderCompilationManager::RunCompileJobs()
 				if (GShaderCompilingManager->ShouldRecompileToDumpShaderDebugInfo(CurrentJob))
 				{
 					CurrentJob.Input.DumpDebugInfoPath = GShaderCompilingManager->CreateShaderDebugInfoPath(CurrentJob.Input);
+					CurrentJob.Output = FShaderCompilerOutput();
 					Compiler->CompileShader(Format, CurrentJob.Input, CurrentJob.Output, FString(FPlatformProcess::ShaderDir()));
 				}
 
