ComputeFramework: Fix reset of other target platforms at the start of BeginCacheForCookedPlatformData()

diff --git a/Engine/Plugins/Runtime/ComputeFramework/Source/ComputeFramework/Private/ComputeGraph.cpp b/Engine/Plugins/Runtime/ComputeFramework/Source/ComputeFramework/Private/ComputeGraph.cpp
--- a/Engine/Plugins/Runtime/ComputeFramework/Source/ComputeFramework/Private/ComputeGraph.cpp
+++ b/Engine/Plugins/Runtime/ComputeFramework/Source/ComputeFramework/Private/ComputeGraph.cpp
@@ -697,7 +697,8 @@ void UComputeGraph::BeginCacheForCookedPlatformData(ITargetPlatform const* Targe
 
 	for (int32 KernelIndex = 0; KernelIndex < KernelInvocations.Num(); ++KernelIndex)
 	{
-		KernelResources[KernelIndex].CachedKernelResourcesForCooking.Reset();
+		TArray< TUniquePtr<FComputeKernelResource> >& Resources = KernelResources[KernelIndex].CachedKernelResourcesForCooking.FindOrAdd(TargetPlatform);
+		Resources.Reset();
 
 		UComputeKernelSource* KernelSource = KernelInvocations[KernelIndex] != nullptr ? KernelInvocations[KernelIndex]->KernelSource : nullptr;
 		if (KernelSource == nullptr)
@@ -717,8 +718,6 @@ void UComputeGraph::BeginCacheForCookedPlatformData(ITargetPlatform const* Targe
 			FString ShaderFriendlyName = GetOuter()->GetName() + TEXT("_") + ShaderEntryPoint;
 			FString ShaderSource = BuildKernelSource(KernelIndex, *KernelSource, AdditionalSources, ShaderHashKey, *ShaderDefinitionSet, *ShaderPermutationVector);
 
-			TArray< TUniquePtr<FComputeKernelResource> >& Resources = KernelResources[KernelIndex].CachedKernelResourcesForCooking.FindOrAdd(TargetPlatform);
-
 			for (int32 ShaderFormatIndex = 0; ShaderFormatIndex < ShaderFormats.Num(); ++ShaderFormatIndex)
 			{
 				TUniquePtr<FShaderParametersMetadataAllocations> ShaderParameterMetadataAllocations = MakeUnique<FShaderParametersMetadataAllocations>();
@@ -796,10 +795,7 @@ void UComputeGraph::ClearCachedCookedPlatformData(ITargetPlatform const* TargetP
 
 void UComputeGraph::ClearAllCachedCookedPlatformData()
 {
-	for (int32 KernelIndex = 0; KernelIndex < KernelInvocations.Num(); ++KernelIndex)
-	{
-		KernelResources[KernelIndex].CachedKernelResourcesForCooking.Reset();
-	}
+	KernelResources.Reset();
 }
 
 #endif // WITH_EDITOR
