ComputeFramework: Set log categories and ensure that graph load/compilation is sent to display.
Helps with debugging cooking issues.

diff --git a/Engine/Plugins/Runtime/ComputeFramework/Source/ComputeFramework/Private/ComputeKernelShaderCompilationManager.cpp b/Engine/Plugins/Runtime/ComputeFramework/Source/ComputeFramework/Private/ComputeKernelShaderCompilationManager.cpp
--- a/Engine/Plugins/Runtime/ComputeFramework/Source/ComputeFramework/Private/ComputeKernelShaderCompilationManager.cpp
+++ b/Engine/Plugins/Runtime/ComputeFramework/Source/ComputeFramework/Private/ComputeKernelShaderCompilationManager.cpp
@@ -5,6 +5,7 @@
 #include "HAL/FileManager.h"
 #include "Misc/Paths.h"
 #include "ComputeFramework/ComputeKernelShared.h"
+#include "RHIShaderFormatDefinitions.inl"
 #include "ShaderCompiler.h"
 
 #if WITH_EDITOR
@@ -299,16 +300,16 @@ void FComputeKernelShaderCompilationManager::ProcessCompiledComputeKernelShaderM
 
 					if (CurrentJob.Output.Errors.Num())
 					{
-						UE_LOG(LogShaders, Log, TEXT("There were errors for job \"%s\""), *CurrentJob.Input.DebugGroupName)
-							for (const FShaderCompilerError& Error : CurrentJob.Output.Errors)
-							{
-								UE_LOG(LogShaders, Log, TEXT("Error: %s"), *Error.GetErrorString())
-							}
+						UE_LOG(LogComputeKernelShaderCompiler, Log, TEXT("There were errors for job \"%s\""), *CurrentJob.Input.DebugGroupName)
+						for (const FShaderCompilerError& Error : CurrentJob.Output.Errors)
+						{
+							UE_LOG(LogComputeKernelShaderCompiler, Log, TEXT("Error: %s"), *Error.GetErrorString())
+						}
 					}
 				}
 				else
 				{
-					UE_LOG(LogShaders, Log, TEXT("There were NO errors for job \"%s\""), *CurrentJob.Input.DebugGroupName);
+					UE_LOG(LogComputeKernelShaderCompiler, Log, TEXT("There were NO errors for job \"%s\""), *CurrentJob.Input.DebugGroupName);
 				}
 			}
 
@@ -394,10 +395,17 @@ void FComputeKernelShaderCompilationManager::ProcessCompiledComputeKernelShaderM
 					{
 						if (CompletedShaderMap->IsComplete(Kernel, true))
 						{
-							const FString SuccessMessage = FString::Printf(TEXT("%s: Shader compilation success!"), *Kernel->GetFriendlyName());
-							const FString FailMessage = FString::Printf(TEXT("%s: Shader compilation failed."), *Kernel->GetFriendlyName());
-
-							Kernel->NotifyCompilationFinished(bSuccess ? SuccessMessage : FailMessage);
+							const FString PlatformName = ShaderPlatformToShaderFormatName(ShaderMap->GetShaderPlatform()).ToString();
+							if (bSuccess)
+							{
+								const FString SuccessMessage = FString::Printf(TEXT("%s: %s shader compilation success!"), *Kernel->GetFriendlyName(), *PlatformName);
+								Kernel->NotifyCompilationFinished(SuccessMessage);
+							}
+							else
+							{
+								const FString FailMessage = FString::Printf(TEXT("%s: %s shader compilation failed."), *Kernel->GetFriendlyName(), *PlatformName);
+								Kernel->NotifyCompilationFinished(FailMessage);
+							}
 						}
 					}
 				}
@@ -429,11 +437,18 @@ void FComputeKernelShaderCompilationManager::ProcessCompiledComputeKernelShaderM
 					Kernel->SetRenderingThreadShaderMap(ShaderMap);
 				});
 
-			const bool bSuccess = ShaderMap && ShaderMap->CompiledSuccessfully();
-			const FString SuccessMessage = FString::Printf(TEXT("%s: Shader compilation success!"), *Kernel->GetFriendlyName());
-			const FString FailMessage = FString::Printf(TEXT("%s: Shader compilation failed."), *Kernel->GetFriendlyName());
-
-			Kernel->NotifyCompilationFinished(bSuccess ? SuccessMessage : FailMessage);
+			if (ShaderMap && ShaderMap->CompiledSuccessfully())
+			{
+				const FString SuccessMessage = FString::Printf(TEXT("%s: %s shader compilation success!"), 
+					*Kernel->GetFriendlyName(), 
+					*ShaderPlatformToShaderFormatName(ShaderMap->GetShaderPlatform()).ToString());
+				Kernel->NotifyCompilationFinished(SuccessMessage);
+			}
+			else
+			{
+				const FString FailMessage = FString::Printf(TEXT("%s: Shader compilation failed."), *Kernel->GetFriendlyName());
+				Kernel->NotifyCompilationFinished(FailMessage);
+			}
 		}
 	}
 #endif
diff --git a/Engine/Plugins/Runtime/ComputeFramework/Source/ComputeFramework/Private/ComputeKernelShaderMap.cpp b/Engine/Plugins/Runtime/ComputeFramework/Source/ComputeFramework/Private/ComputeKernelShaderMap.cpp
--- a/Engine/Plugins/Runtime/ComputeFramework/Source/ComputeFramework/Private/ComputeKernelShaderMap.cpp
+++ b/Engine/Plugins/Runtime/ComputeFramework/Source/ComputeFramework/Private/ComputeKernelShaderMap.cpp
@@ -5,6 +5,7 @@
 #include "ComputeFramework/ComputeKernelShared.h"
 #include "ComputeFramework/ComputeKernelDerivedDataVersion.h"
 #include "ComputeFramework/ComputeKernelShaderCompilationManager.h"
+#include "ComputeFramework/ComputeFramework.h"
 #include "DerivedDataCacheInterface.h"
 #include "Interfaces/ITargetPlatformManagerModule.h"
 #include "ProfilingDebugging/CookStats.h"
@@ -152,13 +153,13 @@ void FComputeKernelShaderType::BeginCompileShader(
 	NewJob->Input.VirtualSourceFilePath = virtualSourcePath;
 	NewJob->Input.EntryPointName = InKernel->GetEntryPoint();
 	NewJob->Input.Environment.IncludeVirtualPathToContentsMap.Add(virtualSourcePath, InKernel->GetHLSLSource());
-	UE_LOG(LogShaders, Verbose, TEXT("%s"), *InKernel->GetHLSLSource());
+	UE_LOG(LogComputeFramework, Verbose, TEXT("%s"), *InKernel->GetHLSLSource());
 	
 	AddReferencedUniformBufferIncludes(NewJob->Input.Environment, NewJob->Input.SourceFilePrefix, InPlatform);
 
 	FShaderCompilerEnvironment& ShaderEnvironment = NewJob->Input.Environment;
 
-	UE_LOG(LogShaders, Verbose, TEXT("			%s"), GetName());
+	UE_LOG(LogComputeFramework, Verbose, TEXT("			%s"), GetName());
 	COOK_STAT(ComputeKernelShaderCookStats::ShadersCompiled++);
 
 	InKernel->SetupCompileEnvironment(PermutationId, ShaderEnvironment);
@@ -319,7 +320,7 @@ void FComputeKernelShaderMap::Compile(
 {
 	if (FPlatformProperties::RequiresCookedData())
 	{
-		UE_LOG(LogShaders, Fatal, TEXT("Trying to compile ComputeKernel shader %s at run-time, which is not supported on consoles!"), *InKernel->GetFriendlyName() );
+		UE_LOG(LogComputeFramework, Fatal, TEXT("Trying to compile ComputeKernel shader %s at run-time, which is not supported on consoles!"), *InKernel->GetFriendlyName() );
 	}
 #if WITH_EDITOR
 	else
@@ -340,7 +341,7 @@ void FComputeKernelShaderMap::Compile(
 		{
 			// Assign a unique identifier so that shaders from this shader map can be associated with it after a deferred compile
 			CompilingId = NextCompilingId;
-			UE_LOG(LogShaders, Log, TEXT("CompilingId = %p %d"), InKernel, CompilingId);
+			UE_LOG(LogComputeFramework, Log, TEXT("CompilingId = %p %d"), InKernel, CompilingId);
 			InKernel->AddCompileId(CompilingId);
 
 			check(NextCompilingId < UINT_MAX);
@@ -400,15 +401,13 @@ void FComputeKernelShaderMap::Compile(
 					InKernel->RemoveOutstandingCompileId(CompilingId);
 
 					FString Message = FString::Printf(TEXT("%s: Skipping compilation because it isn't supported on this target type."), *InKernel->GetFriendlyName());
-					UE_LOG(LogShaders, Display, TEXT("%s"), *Message);
-					
 					InKernel->NotifyCompilationFinished(Message);
 				}
 			}
   
 			if (!Kernel)
 			{
-				UE_LOG(LogShaders, Log, TEXT("		%u Shaders"), NumShaders);
+				UE_LOG(LogComputeFramework, Log, TEXT("		%u Shaders"), NumShaders);
 			}
 
 			// Register this shader map in the global ComputeKernel->shadermap map
@@ -496,7 +495,7 @@ bool FComputeKernelShaderMap::TryToAddToExistingCompilationTask(FComputeKernelRe
 	{
 		CorrespondingKernels->AddUnique(InKernel);
 
-		UE_LOG(LogShaders, Log, TEXT("TryToAddToExistingCompilationTask %p %d"), InKernel, GetCompilingId());
+		UE_LOG(LogComputeFramework, Log, TEXT("TryToAddToExistingCompilationTask %p %d"), InKernel, GetCompilingId());
 
 #if DEBUG_INFINITESHADERCOMPILE
 		UE_LOG(LogTemp, Display, TEXT("Added shader map 0x%08X%08X from ComputeKernel 0x%08X%08X"), (int)((int64)(this) >> 32), (int)((int64)(this)), (int)((int64)(InKernel) >> 32), (int)((int64)(InKernel)));
@@ -518,7 +517,7 @@ bool FComputeKernelShaderMap::IsComputeKernelShaderComplete(const FComputeKernel
 			{
 				if (!bSilent)
 				{
-					UE_LOG(LogShaders, Warning, TEXT("Incomplete shader %s, missing FComputeKernelShader %s."), *InKernel->GetFriendlyName(), InShaderType->GetName());
+					UE_LOG(LogComputeFramework, Warning, TEXT("Incomplete shader %s, missing FComputeKernelShader %s."), *InKernel->GetFriendlyName(), InShaderType->GetName());
 				}
 				return false;
 			}
diff --git a/Engine/Plugins/Runtime/ComputeFramework/Source/ComputeFramework/Private/ComputeKernelShared.cpp b/Engine/Plugins/Runtime/ComputeFramework/Source/ComputeFramework/Private/ComputeKernelShared.cpp
--- a/Engine/Plugins/Runtime/ComputeFramework/Source/ComputeFramework/Private/ComputeKernelShared.cpp
+++ b/Engine/Plugins/Runtime/ComputeFramework/Source/ComputeFramework/Private/ComputeKernelShared.cpp
@@ -8,9 +8,11 @@
 #include "ComputeFramework/ComputeKernelShader.h"
 #include "ComputeFramework/ComputeKernelShaderCompilationManager.h"
 #include "ComputeFramework/ComputeKernelSource.h"
+#include "ComputeFramework/ComputeFramework.h"
 #include "ComputeFramework/ShaderParameterMetadataAllocation.h"
 #include "Interfaces/ITargetPlatform.h"
 #include "RendererInterface.h"
+#include "RHIShaderFormatDefinitions.inl"
 #include "ShaderCompiler.h"
 #include "ShaderParameterMetadataBuilder.h"
 #include "Stats/StatsMisc.h"
@@ -53,6 +55,7 @@ bool FComputeKernelResource::ShouldCache(EShaderPlatform InPlatform, const FShad
 
 void FComputeKernelResource::NotifyCompilationFinished(FString const& ResultMessage)
 {
+	UE_LOG(LogComputeFramework, Display, TEXT("%s"), *ResultMessage);
 	CompileOutputMessages.Add(ResultMessage);
 	OnCompilationCompleteDelegate.ExecuteIfBound(this);
 }
@@ -66,8 +69,6 @@ void FComputeKernelResource::CancelCompilation()
 		OutstandingCompileShaderMapIds.Empty();
 
 		FString Message = FString::Printf(TEXT("%s: Compilation cancelled."), *GetFriendlyName());
-		UE_LOG(LogShaders, Log, TEXT("%s"), *Message);
-
 		NotifyCompilationFinished(Message);
 	}
 #endif
@@ -77,7 +78,7 @@ void FComputeKernelResource::RemoveOutstandingCompileId(const int32 InOldOutstan
 {
 	if (0 <= OutstandingCompileShaderMapIds.Remove(InOldOutstandingCompileShaderMapId))
 	{
-		UE_LOG(LogShaders, Log, TEXT("RemoveOutstandingCompileId %p %d"), this, InOldOutstandingCompileShaderMapId);
+		UE_LOG(LogComputeFramework, Log, TEXT("RemoveOutstandingCompileId %p %d"), this, InOldOutstandingCompileShaderMapId);
 	}
 }
 
@@ -135,7 +136,7 @@ void FComputeKernelResource::GetShaderMapId(EShaderPlatform InPlatform, const IT
 #else
 		if (TargetPlatform != nullptr)
 		{
-			UE_LOG(LogShaders, Error, TEXT("FComputeKernelResource::GetShaderMapId: TargetPlatform is not null, but a cooked executable cannot target platforms other than its own."));
+			UE_LOG(LogComputeFramework, Error, TEXT("FComputeKernelResource::GetShaderMapId: TargetPlatform is not null, but a cooked executable cannot target platforms other than its own."));
 		}
 		OutId.LayoutParams.InitializeForCurrent();
 #endif
@@ -170,7 +171,7 @@ void FComputeKernelResource::SerializeShaderMap(FArchive& Ar)
 
 	if (FPlatformProperties::RequiresCookedData() && !bCooked && Ar.IsLoading())
 	{
-		UE_LOG(LogShaders, Fatal, TEXT("This platform requires cooked packages, and shaders were not cooked into this kernel %s."), *GetFriendlyName());
+		UE_LOG(LogComputeFramework, Fatal, TEXT("This platform requires cooked packages, and shaders were not cooked into this kernel %s."), *GetFriendlyName());
 	}
 
 	if (bCooked)
@@ -326,12 +327,11 @@ bool FComputeKernelResource::CacheShaders(const FComputeKernelShaderMapId& InSha
 				if (GameThreadShaderMap && GameThreadShaderMap->IsValid())
 				{
 					FString Message = FString::Printf(TEXT("Loaded shaders for %s from DDC."), *GetFriendlyName());
-					UE_LOG(LogTemp, Display, TEXT("%s"), *Message);
 					CompileOutputMessages.Add(Message);
 				}
 				else
 				{
-					UE_LOG(LogTemp, Display, TEXT("Loading shader for kernel %s from DDC failed. Shader needs recompile."), *GetFriendlyName());
+					UE_LOG(LogComputeFramework, Log, TEXT("Loading shader for kernel %s from DDC failed. Shader needs recompile."), *GetFriendlyName());
 				}
 			}
 		}
@@ -345,10 +345,10 @@ bool FComputeKernelResource::CacheShaders(const FComputeKernelShaderMapId& InSha
 	if (GameThreadShaderMap && GameThreadShaderMap->TryToAddToExistingCompilationTask(this))
 	{
 #if DEBUG_INFINITESHADERCOMPILE
-		UE_LOG(LogTemp, Display, TEXT("Found existing compiling shader for kernel %s, linking to other GameThreadShaderMap 0x%08X%08X"), *GetFriendlyName(), (int)((int64)(GameThreadShaderMap.GetReference()) >> 32), (int)((int64)(GameThreadShaderMap.GetReference())));
+		UE_LOG(LogComputeFramework, Display, TEXT("Found existing compiling shader for kernel %s, linking to other GameThreadShaderMap 0x%08X%08X"), *GetFriendlyName(), (int)((int64)(GameThreadShaderMap.GetReference()) >> 32), (int)((int64)(GameThreadShaderMap.GetReference())));
 #endif
 		OutstandingCompileShaderMapIds.AddUnique(GameThreadShaderMap->GetCompilingId());
-		UE_LOG(LogShaders, Log, TEXT("CacheShaders AddUniqueExisting %p %d"), this, GameThreadShaderMap->GetCompilingId());
+		UE_LOG(LogComputeFramework, Log, TEXT("CacheShaders AddUniqueExisting %p %d"), this, GameThreadShaderMap->GetCompilingId());
 
 		GameThreadShaderMap = nullptr;
 		bSucceeded = true;
@@ -357,12 +357,12 @@ bool FComputeKernelResource::CacheShaders(const FComputeKernelShaderMapId& InSha
 	{
 		if (bContainsInlineShaders || FPlatformProperties::RequiresCookedData())
 		{
-			UE_LOG(LogShaders, Log, TEXT("Can't compile %s with cooked content!"), *GetFriendlyName());
+			UE_LOG(LogComputeFramework, Log, TEXT("Can't compile %s with cooked content!"), *GetFriendlyName());
 			GameThreadShaderMap = nullptr;
 		}
 		else
 		{
-			UE_LOG(LogShaders, Log, TEXT("%s cached shader map for kernel %s, compiling."), GameThreadShaderMap ? TEXT("Incomplete") : TEXT("Missing"), *GetFriendlyName());
+			UE_LOG(LogComputeFramework, Log, TEXT("%s cached shader map for kernel %s, compiling."), GameThreadShaderMap ? TEXT("Incomplete") : TEXT("Missing"), *GetFriendlyName());
 
 			// If there's no cached shader map for this kernel compile a new one.
 			// This is just kicking off the compile, GameThreadShaderMap will not be complete yet
@@ -376,7 +376,7 @@ bool FComputeKernelResource::CacheShaders(const FComputeKernelShaderMapId& InSha
 	}
 	else
 	{
-		NotifyCompilationFinished(FString::Printf(TEXT("%s: Loaded shaders from DDC."), *GetFriendlyName()));
+		NotifyCompilationFinished(FString::Printf(TEXT("%s: Loaded %s shaders from DDC."), *GetFriendlyName(), *ShaderPlatformToShaderFormatName(InPlatform).ToString()));
 		bSucceeded = true;
 	}
 
@@ -401,7 +401,7 @@ void FComputeKernelResource::FinishCompilation()
 	{
 		for (int32 i = 0; i < ShaderMapIdsToFinish.Num(); i++)
 		{
-			UE_LOG(LogShaders, Log, TEXT("FinishCompilation()[%d] %s id %d!"), i, *GetFriendlyName(), ShaderMapIdsToFinish[i]);
+			UE_LOG(LogComputeFramework, Log, TEXT("FinishCompilation()[%d] %s id %d!"), i, *GetFriendlyName(), ShaderMapIdsToFinish[i]);
 		}
 	
 		// Block until the shader maps that we will save have finished being compiled
@@ -480,7 +480,7 @@ bool FComputeKernelResource::BeginCompileShaderMap(const FComputeKernelShaderMap
 	}
 	else
 	{
-		UE_LOG(LogShaders, Log, TEXT("BeginCompileShaderMap AddUnique %p %d"), this, NewShaderMap->GetCompilingId());
+		UE_LOG(LogComputeFramework, Log, TEXT("BeginCompileShaderMap AddUnique %p %d"), this, NewShaderMap->GetCompilingId());
 		OutstandingCompileShaderMapIds.AddUnique(NewShaderMap->GetCompilingId());
 		
 		// Async compile, use nullptr to detect it if used
@@ -489,7 +489,7 @@ bool FComputeKernelResource::BeginCompileShaderMap(const FComputeKernelShaderMap
 
 	return true;
 #else
-	UE_LOG(LogShaders, Fatal, TEXT("Compiling of shaders in a build without editordata is not supported."));
+	UE_LOG(LogComputeFramework, Fatal, TEXT("Compiling of shaders in a build without editordata is not supported."));
 	return false;
 #endif
 }
