Don't write shadermap to renderthread when cooking.

diff --git a/Engine/Plugins/Runtime/ComputeFramework/Source/ComputeFramework/Private/ComputeKernelShared.cpp b/Engine/Plugins/Runtime/ComputeFramework/Source/ComputeFramework/Private/ComputeKernelShared.cpp
--- a/Engine/Plugins/Runtime/ComputeFramework/Source/ComputeFramework/Private/ComputeKernelShared.cpp
+++ b/Engine/Plugins/Runtime/ComputeFramework/Source/ComputeFramework/Private/ComputeKernelShared.cpp
@@ -380,13 +380,16 @@ bool FComputeKernelResource::CacheShaders(const FComputeKernelShaderMapId& InSha
 		bSucceeded = true;
 	}
 
-	FComputeKernelResource* Kernel = this;
-	FComputeKernelShaderMap* LoadedShaderMap = GameThreadShaderMap;
-	ENQUEUE_RENDER_COMMAND(FSetShaderMapOnComputeKernel)(
-		[Kernel, LoadedShaderMap](FRHICommandListImmediate& RHICmdList)
-		{
-			Kernel->SetRenderingThreadShaderMap(LoadedShaderMap);
-		});
+	if (bApplyCompletedShaderMapForRendering)
+	{
+		FComputeKernelResource* Kernel = this;
+		FComputeKernelShaderMap* LoadedShaderMap = GameThreadShaderMap;
+		ENQUEUE_RENDER_COMMAND(FSetShaderMapOnComputeKernel)(
+			[Kernel, LoadedShaderMap](FRHICommandListImmediate& RHICmdList)
+			{
+				Kernel->SetRenderingThreadShaderMap(LoadedShaderMap);
+			});
+	}
 
 	return bSucceeded;
 }
