ComputeFramework: Disable on Vulkan platforms using DataDrivenPlatfomInfo
Do this until we fix user defined structures that don't obey Vulkan packing rules.
Avoids crashes on some drivers.

diff --git a/Engine/Config/Windows/DataDrivenPlatformInfo.ini b/Engine/Config/Windows/DataDrivenPlatformInfo.ini
--- a/Engine/Config/Windows/DataDrivenPlatformInfo.ini
+++ b/Engine/Config/Windows/DataDrivenPlatformInfo.ini
@@ -147,6 +147,7 @@ ShaderFormat=PCD3D_ES31
 bIsPC=true
 bIsMobile=true
 bSupportsDistanceFields = true
+bSupportsComputeFramework=true
 bSupportsByteBufferComputeShaders = true
 bSupportsInstancedStereo=true
 bSupportsManualVertexFetch=false
diff --git a/Engine/Plugins/Runtime/ComputeFramework/Shaders/Private/BufferAlias.ush b/Engine/Plugins/Runtime/ComputeFramework/Shaders/Private/BufferAlias.ush
--- a/Engine/Plugins/Runtime/ComputeFramework/Shaders/Private/BufferAlias.ush
+++ b/Engine/Plugins/Runtime/ComputeFramework/Shaders/Private/BufferAlias.ush
@@ -8,44 +8,38 @@
 
 void ReadRawBuffer(in StructuredBuffer<float> InBuffer, in uint InBufferIndex, out float3 OutValue)
 {
-	for (uint Index = 0; Index < 3; ++Index)
-	{
-		OutValue[Index] = InBuffer[InBufferIndex * 3 + Index];
-	}
+	OutValue.x = InBuffer[InBufferIndex * 3 + 0];
+	OutValue.y = InBuffer[InBufferIndex * 3 + 1];
+	OutValue.z = InBuffer[InBufferIndex * 3 + 2];
 }
 void ReadRawBuffer(in RWStructuredBuffer<float> InBuffer, in uint InBufferIndex, out float3 OutValue)
 {
-	for (uint Index = 0; Index < 3; ++Index)
-	{
-		OutValue[Index] = InBuffer[InBufferIndex * 3 + Index];
-	}
+	OutValue.x = InBuffer[InBufferIndex * 3 + 0];
+	OutValue.y = InBuffer[InBufferIndex * 3 + 1];
+	OutValue.z = InBuffer[InBufferIndex * 3 + 2];
 }
 void WriteRawBuffer(in RWStructuredBuffer<float> InBuffer, in uint InBufferIndex, in float3 InValue)
 {
-	for (uint Index = 0; Index < 3; ++Index)
-	{
-		InBuffer[InBufferIndex * 3 + Index] = InValue[Index];
-	}
+	InBuffer[InBufferIndex * 3 + 0] = InValue.x;
+	InBuffer[InBufferIndex * 3 + 1] = InValue.y;
+	InBuffer[InBufferIndex * 3 + 2] = InValue.z;
 }
 
 void ReadRawBuffer(in StructuredBuffer<int> InBuffer, in uint InBufferIndex, out int3 OutValue)
 {
-	for (uint Index = 0; Index < 3; ++Index)
-	{
-		OutValue[Index] = InBuffer[InBufferIndex * 3 + Index];
-	}
+	OutValue.x = InBuffer[InBufferIndex * 3 + 0];
+	OutValue.y = InBuffer[InBufferIndex * 3 + 1];
+	OutValue.z = InBuffer[InBufferIndex * 3 + 2];
 }
 void ReadRawBuffer(in RWStructuredBuffer<int> InBuffer, in uint InBufferIndex, out int3 OutValue)
 {
-	for (uint Index = 0; Index < 3; ++Index)
-	{
-		OutValue[Index] = InBuffer[InBufferIndex * 3 + Index];
-	}
+	OutValue.x = InBuffer[InBufferIndex * 3 + 0];
+	OutValue.y = InBuffer[InBufferIndex * 3 + 1];
+	OutValue.z = InBuffer[InBufferIndex * 3 + 2];
 }
 void WriteRawBuffer(in RWStructuredBuffer<int> InBuffer, in uint InBufferIndex, in int3 InValue)
 {
-	for (uint Index = 0; Index < 3; ++Index)
-	{
-		InBuffer[InBufferIndex * 3 + Index] = InValue[Index];
-	}
+	InBuffer[InBufferIndex * 3 + 0] = InValue.x;
+	InBuffer[InBufferIndex * 3 + 1] = InValue.y;
+	InBuffer[InBufferIndex * 3 + 2] = InValue.z;
 }
diff --git a/Engine/Plugins/Runtime/ComputeFramework/Source/ComputeFramework/Private/ComputeFramework.cpp b/Engine/Plugins/Runtime/ComputeFramework/Source/ComputeFramework/Private/ComputeFramework.cpp
--- a/Engine/Plugins/Runtime/ComputeFramework/Source/ComputeFramework/Private/ComputeFramework.cpp
+++ b/Engine/Plugins/Runtime/ComputeFramework/Source/ComputeFramework/Private/ComputeFramework.cpp
@@ -26,9 +26,14 @@ FAutoConsoleCommand CmdRebuildComputeGraphs(
 
 namespace ComputeFramework
 {
-	bool IsEnabled(EShaderPlatform ShaderPlatform)
+	bool IsSupported(EShaderPlatform ShaderPlatform)
 	{
-		return (GComputeFrameworkEnable > 0)	&& FDataDrivenShaderPlatformInfo::GetSupportsComputeFramework(ShaderPlatform);
+		return FDataDrivenShaderPlatformInfo::GetSupportsComputeFramework(ShaderPlatform);
+	}
+
+	bool IsEnabled()
+	{
+		return GComputeFrameworkEnable > 0;
 	}
 
 	void RebuildComputeGraphs()
diff --git a/Engine/Plugins/Runtime/ComputeFramework/Source/ComputeFramework/Private/ComputeGraphInstance.cpp b/Engine/Plugins/Runtime/ComputeFramework/Source/ComputeFramework/Private/ComputeGraphInstance.cpp
--- a/Engine/Plugins/Runtime/ComputeFramework/Source/ComputeFramework/Private/ComputeGraphInstance.cpp
+++ b/Engine/Plugins/Runtime/ComputeFramework/Source/ComputeFramework/Private/ComputeGraphInstance.cpp
@@ -4,6 +4,7 @@
 
 #include "ComputeFramework/ComputeDataInterface.h"
 #include "ComputeFramework/ComputeDataProvider.h"
+#include "ComputeFramework/ComputeFramework.h"
 #include "ComputeFramework/ComputeFrameworkModule.h"
 #include "ComputeFramework/ComputeGraph.h"
 #include "ComputeFramework/ComputeGraphRenderProxy.h"
@@ -36,7 +37,7 @@ bool FComputeGraphInstance::EnqueueWork(UComputeGraph* InComputeGraph, FSceneInt
 		return false;
 	}
 
-	if (FComputeFrameworkModule::GetComputeSystem() == nullptr)
+	if (!ComputeFramework::IsEnabled() || FComputeFrameworkModule::GetComputeSystem() == nullptr)
 	{
 		return false;
 	}
diff --git a/Engine/Plugins/Runtime/ComputeFramework/Source/ComputeFramework/Private/ComputeKernelShaderMap.cpp b/Engine/Plugins/Runtime/ComputeFramework/Source/ComputeFramework/Private/ComputeKernelShaderMap.cpp
--- a/Engine/Plugins/Runtime/ComputeFramework/Source/ComputeFramework/Private/ComputeKernelShaderMap.cpp
+++ b/Engine/Plugins/Runtime/ComputeFramework/Source/ComputeFramework/Private/ComputeKernelShaderMap.cpp
@@ -10,6 +10,7 @@
 #include "Interfaces/ITargetPlatformManagerModule.h"
 #include "ProfilingDebugging/CookStats.h"
 #include "ProfilingDebugging/DiagnosticTable.h"
+#include "RHIShaderFormatDefinitions.inl"
 #include "Serialization/MemoryWriter.h"
 #include "Serialization/MemoryReader.h"
 #include "ShaderCompiler.h"
@@ -347,12 +348,6 @@ void FComputeKernelShaderMap::Compile(
 			check(NextCompilingId < UINT_MAX);
 			NextCompilingId++;
   
-			TArray<FComputeKernelResource*> NewCorrespondingKernels;
-			NewCorrespondingKernels.Add(InKernel);
-			ComputeKernelShaderMapsBeingCompiled.Add(this, NewCorrespondingKernels);
-#if DEBUG_INFINITESHADERCOMPILE
-			UE_LOG(LogTemp, Display, TEXT("Added ComputeKernel ShaderMap 0x%08X%08X with kernel 0x%08X%08X to ComputeKernelShaderMapsBeingCompiled"), (int)((int64)(this) >> 32), (int)((int64)(this)), (int)((int64)(InKernel) >> 32), (int)((int64)(InKernel)));
-#endif  
 			// Setup the compilation environment.
 			InKernel->SetupShaderCompilationEnvironment(InPlatform, *InCompilationEnvironment);
   
@@ -400,7 +395,9 @@ void FComputeKernelShaderMap::Compile(
 				{
 					InKernel->RemoveOutstandingCompileId(CompilingId);
 
-					FString Message = FString::Printf(TEXT("%s: Skipping compilation because it isn't supported on this target type."), *InKernel->GetFriendlyName());
+					FString Message = FString::Printf(TEXT("%s: Compilation not supported on %s."), 
+						*InKernel->GetFriendlyName(), 
+						*ShaderPlatformToShaderFormatName(InPlatform).ToString());
 					InKernel->NotifyCompilationFinished(Message);
 				}
 			}
@@ -419,7 +416,17 @@ void FComputeKernelShaderMap::Compile(
 			// Mark as not having been compiled
 			bCompiledSuccessfully = false;
   
-			GComputeKernelShaderCompilationManager.AddJobs(NewJobs);
+			if (NumShaders > 0)
+			{
+				GComputeKernelShaderCompilationManager.AddJobs(NewJobs);
+
+				TArray<FComputeKernelResource*> NewCorrespondingKernels;
+				NewCorrespondingKernels.Add(InKernel);
+				ComputeKernelShaderMapsBeingCompiled.Add(this, NewCorrespondingKernels);
+#if DEBUG_INFINITESHADERCOMPILE
+				UE_LOG(LogTemp, Display, TEXT("Added ComputeKernel ShaderMap 0x%08X%08X with kernel 0x%08X%08X to ComputeKernelShaderMapsBeingCompiled"), (int)((int64)(this) >> 32), (int)((int64)(this)), (int)((int64)(InKernel) >> 32), (int)((int64)(InKernel)));
+#endif  
+			}
   
 			// Compile the shaders for this shader map now if not deferring and deferred compiles are not enabled globally
 			if (bSynchronousCompile)
diff --git a/Engine/Plugins/Runtime/ComputeFramework/Source/ComputeFramework/Public/ComputeFramework/ComputeFramework.h b/Engine/Plugins/Runtime/ComputeFramework/Source/ComputeFramework/Public/ComputeFramework/ComputeFramework.h
--- a/Engine/Plugins/Runtime/ComputeFramework/Source/ComputeFramework/Public/ComputeFramework/ComputeFramework.h
+++ b/Engine/Plugins/Runtime/ComputeFramework/Source/ComputeFramework/Public/ComputeFramework/ComputeFramework.h
@@ -9,8 +9,11 @@ DECLARE_LOG_CATEGORY_EXTERN(LogComputeFramework, Log, All);
 
 namespace ComputeFramework
 {
-	/** Returns true if ComputeFramework is supported. */
-	COMPUTEFRAMEWORK_API bool IsEnabled(EShaderPlatform ShaderPlatform);
+	/** Returns true if ComputeFramework is supported on a platform. */
+	COMPUTEFRAMEWORK_API bool IsSupported(EShaderPlatform ShaderPlatform);
+
+	/** Returns true if ComputeFramework is currently enabled. */
+	COMPUTEFRAMEWORK_API bool IsEnabled();
 
 	/** Rebuild all compute graphs. */
 	COMPUTEFRAMEWORK_API void RebuildComputeGraphs();
