ComputeFramework: Disable on Vulkan platforms using DataDrivenPlatfomInfo
Do this until we fix user defined structures that don't obey Vulkan packing rules.
Avoids crashes on some drivers.

diff --git a/Engine/Plugins/Runtime/ComputeFramework/Source/ComputeFramework/Private/ComputeKernelShared.cpp b/Engine/Plugins/Runtime/ComputeFramework/Source/ComputeFramework/Private/ComputeKernelShared.cpp
--- a/Engine/Plugins/Runtime/ComputeFramework/Source/ComputeFramework/Private/ComputeKernelShared.cpp
+++ b/Engine/Plugins/Runtime/ComputeFramework/Source/ComputeFramework/Private/ComputeKernelShared.cpp
@@ -49,8 +49,8 @@ void FComputeKernelResource::SetupShaderCompilationEnvironment(EShaderPlatform I
 
 bool FComputeKernelResource::ShouldCache(EShaderPlatform InPlatform, const FShaderType* InShaderType) const
 {
-	check(InShaderType->GetComputeKernelShaderType() )
-	return true;
+	check(InShaderType->GetComputeKernelShaderType())
+	return ComputeFramework::IsSupported(InPlatform);
 }
 
 void FComputeKernelResource::NotifyCompilationFinished(FString const& ResultMessage)
