ComputeFramework shader compiling WITH_EDITOR cleanup.
None of the shader compiling code runs in non-Editor mode, so don't compile it for those targets.

diff --git a/Engine/Plugins/Runtime/ComputeFramework/Source/ComputeFramework/Private/ComputeFramework/ComputeKernelShaderType.h b/Engine/Plugins/Runtime/ComputeFramework/Source/ComputeFramework/Private/ComputeFramework/ComputeKernelShaderType.h
--- a/Engine/Plugins/Runtime/ComputeFramework/Source/ComputeFramework/Private/ComputeFramework/ComputeKernelShaderType.h
+++ b/Engine/Plugins/Runtime/ComputeFramework/Source/ComputeFramework/Private/ComputeFramework/ComputeKernelShaderType.h
@@ -6,7 +6,9 @@
 #include "Engine/EngineTypes.h"
 #include "GlobalShader.h"
 #include "Shader.h"
+#if WITH_EDITOR
 #include "ShaderCompiler.h"
+#endif
 
 class FShaderCommonCompileJob;
 class FShaderCompileJob;
@@ -98,6 +100,7 @@ public:
 	{
 	}
 
+#if WITH_EDITOR
 	void BeginCompileShader(
 		uint32 ShaderMapId,
 		int32 PermutationId,
@@ -116,6 +119,7 @@ public:
 		const FShaderCompileJob& CurrentJob,
 		const FString& InDebugDescription
 		) const;
+#endif // WITH_EDITOR
 
 	/**
 	 * Checks if the shader type should be cached.
@@ -126,7 +130,7 @@ public:
 		return ShouldCompilePermutation(FComputeKernelShaderPermutationParameters(InPlatform));
 	}
 
-
+#if WITH_EDITOR
 protected:
 	/**
 	 * Sets up the environment used to compile an instance of this shader type.
@@ -137,4 +141,5 @@ protected:
 	{
 		ModifyCompilationEnvironment(FComputeKernelShaderPermutationParameters(InPlatform), OutEnvironment);
 	}
+#endif // WITH_EDITOR
 };
diff --git a/Engine/Plugins/Runtime/ComputeFramework/Source/ComputeFramework/Private/ComputeFramework/ComputeKernelShared.h b/Engine/Plugins/Runtime/ComputeFramework/Source/ComputeFramework/Private/ComputeFramework/ComputeKernelShared.h
--- a/Engine/Plugins/Runtime/ComputeFramework/Source/ComputeFramework/Private/ComputeFramework/ComputeKernelShared.h
+++ b/Engine/Plugins/Runtime/ComputeFramework/Source/ComputeFramework/Private/ComputeFramework/ComputeKernelShared.h
@@ -59,7 +59,9 @@ public:
 
 	~FComputeKernelShaderMapId() = default;
 
+#if WITH_EDITOR
 	void SetShaderDependencies(const TArray<FShaderType*>& InShaderTypes, EShaderPlatform InShaderPlatform);
+#endif
 
 	//void Serialize(FArchive& Ar);
 
@@ -89,8 +91,10 @@ public:
 		return !(*this == InReferenceSet);
 	}
 
+#if WITH_EDITOR
 	/** Appends string representations of this Id to a key string. */
 	void AppendKeyString(FString& OutKeyString) const;
+#endif // WITH_EDITOR
 
 	/** Returns true if the requested shader type is a dependency of this shader map Id. */
 	bool ContainsShaderType(const FShaderType* ShaderType) const;
@@ -135,15 +139,18 @@ public:
 	template<typename ShaderType> TShaderRef<ShaderType> GetShader(int32 PermutationId) const { return TShaderRef<ShaderType>(GetContent()->GetShader<ShaderType>(PermutationId), *this); }
 	TShaderRef<FShader> GetShader(FShaderType* ShaderType, int32 PermutationId) const { return TShaderRef<FShader>(GetContent()->GetShader(ShaderType, PermutationId), *this); }
 
+#if WITH_EDITOR
 	/**
 	 * Attempts to load the shader map for the given kernel from the Derived Data Cache.
 	 * If InOutShaderMap is valid, attempts to load the individual missing shaders instead.
 	 */
 	static void LoadFromDerivedDataCache(const FComputeKernelResource* InKernel, const FComputeKernelShaderMapId& InShaderMapId, EShaderPlatform InPlatform, TRefCountPtr<FComputeKernelShaderMap>& InOutShaderMap);
+#endif // WITH_EDITOR
 
 	FComputeKernelShaderMap();
 	~FComputeKernelShaderMap();
 
+#if WITH_EDITOR
 	/**
 	 * Compiles the shaders for a kernel and caches them in this shader map.
 	 * @param InKernel - The kernel to compile shaders for.
@@ -162,6 +169,7 @@ public:
 
 	/** Sorts the incoming compiled jobs into the appropriate shader maps, and finalizes this shader map so that it can be used for rendering. */
 	bool ProcessCompilationResults(const TArray<FShaderCommonCompileJobPtr>& InCompilationResults, int32& InOutResultIndex, float& InOutTimeBudget);
+#endif // WITH_EDITOR
 
 	/**
 	 * Checks whether the shader map is missing any shader types necessary for the given kernel.
@@ -170,6 +178,7 @@ public:
 	 */
 	bool IsComplete(const FComputeKernelResource* InKernel, bool bSilent);
 
+#if WITH_EDITOR
 	/** Attempts to load missing shaders from memory. */
 	void LoadMissingShadersFromMemory(const FComputeKernelResource* InKernel);
 
@@ -180,6 +189,7 @@ public:
 	 * @return True if the shader map was being compiled and we added it to the list to be applied.
 	 */
 	bool TryToAddToExistingCompilationTask(FComputeKernelResource* InKernel);
+#endif // WITH_EDITOR
 
 	/** Builds a list of the shaders in a shader map. */
 	 void GetShaderList(TMap<FShaderId, TShaderRef<FShader>>& OutShaders) const;
@@ -191,23 +201,18 @@ public:
 	void AddRef();
 	void Release();
 
-	/**
-	 * Removes all entries in the cache with exceptions based on a shader type
-	 * @param ShaderType - The shader type to flush
-	 */
-	void FlushShadersByShaderType(const FShaderType* ShaderType);
-
+#if WITH_EDITOR
 	/** Removes a kernel from ComputeKernelShaderMapsBeingCompiled. */
 	static void RemovePending(FComputeKernelResource* InKernel);
-
-	/** Finds a shader map currently being compiled that was enqueued for the given kernel. */
-	static const FComputeKernelShaderMap* GetShaderMapBeingCompiled(const FComputeKernelResource* InKernel);
+#endif
 
 	/** Serializes the shader map. */
 	bool Serialize(FArchive& Ar, bool bInlineShaderResources = true);
 
+#if WITH_EDITOR
 	/** Saves this shader map to the derived data cache. */
 	void SaveToDerivedDataCache();
+#endif
 
 	// Accessors.
 	const FComputeKernelShaderMapId& GetShaderMapId() const	{ return GetContent()->ShaderMapId; }
@@ -224,9 +229,13 @@ public:
 
 	int32 GetNumRefs() const { return NumRefs; }
 	uint32 GetCompilingId()  { return CompilingId; }
+
+#if WITH_EDITOR
 	static TMap<TRefCountPtr<FComputeKernelShaderMap>, TArray<FComputeKernelResource*> > &GetInFlightShaderMaps() { return ComputeKernelShaderMapsBeingCompiled; }
 
 	void SetCompiledSuccessfully(bool bSuccess) { bCompiledSuccessfully = bSuccess; }
+#endif // WITH_EDITOR
+
 private:
 
 	/**
@@ -242,11 +251,13 @@ private:
 	 */
 	static TArray<FComputeKernelShaderMap*> AllComputeKernelShaderMaps;
 
+#if WITH_EDITOR
 	/** Next value for CompilingId. */
 	static uint32 NextCompilingId;
 
 	/** Tracks resources and their shader maps that need to be compiled but whose compilation is being deferred. */
 	static TMap<TRefCountPtr<FComputeKernelShaderMap>, TArray<FComputeKernelResource*> > ComputeKernelShaderMapsBeingCompiled;
+#endif
 
 	/** Uniquely identifies this shader map during compilation, needed for deferred compilation where shaders from multiple shader maps are compiled together. */
 	uint32 CompilingId;
@@ -272,7 +283,9 @@ private:
 
 	uint32 bHasFrozenContent : 1;
 
+#if WITH_EDITOR
 	FShader* ProcessCompilationResultsForSingleJob(FShaderCompileJob& InSingleJob, const FSHAHash& InShaderMapHash);
+#endif
 
 	bool IsComputeKernelShaderComplete(const FComputeKernelResource* InKernel, const FComputeKernelShaderType* InShaderType, bool bSilent);
 
@@ -331,6 +344,7 @@ public:
 	 */
 	virtual void NotifyCompilationFinished(FString const& ResultMessage);
 
+#if WITH_EDITOR
 	/**
 	 * Cancels all outstanding compilation jobs
 	 */
@@ -340,6 +354,7 @@ public:
 	 * Blocks until compilation has completed. Returns immediately if a compilation is not outstanding.
 	 */
 	void FinishCompilation();
+#endif // WITH_EDITOR
 
 	/**
 	 * Checks if the compilation for this shader is finished
@@ -425,7 +440,9 @@ public:
 		FName const& InAssetPath
 	);
 
+#if WITH_EDITOR
 	void SetupCompileEnvironment(int32 InPermutationId, FShaderCompilerEnvironment& OutShaderEnvironment) const;
+#endif // WITH_EDITOR
 
 	const FShaderParametersMetadata* GetShaderParamMetadata() const
 	{
@@ -442,11 +459,12 @@ public:
 	FOnComputeKernelCompilationComplete& OnCompilationComplete() { return OnCompilationCompleteDelegate; }
 
 protected:
-
+#if WITH_EDITOR
 	/**
 	 * Fills the passed array with IDs of shader maps unfinished compilation jobs.
 	 */
 	void GetShaderMapIDsWithUnfinishedCompilation(TArray<int32>& OutShaderMapIds);
+#endif // WITH_EDITOR
 
 private:
 	TArray<FString> CompileOutputMessages;
@@ -515,6 +533,7 @@ private:
 
 	FOnComputeKernelCompilationComplete OnCompilationCompleteDelegate;
 
+#if WITH_EDITOR
 	/**
 	 * Compiles this kernel for InPlatform, storing the result in OutShaderMap if the compile was synchronous
 	 */
@@ -530,7 +549,7 @@ private:
 		EShaderPlatform InPlatform,
 		FShaderCompilerEnvironment& OutEnvironment
 		) const;
-
+#endif // WITH_EDITOR
 
 	FString FriendlyName;
 
diff --git a/Engine/Plugins/Runtime/ComputeFramework/Source/ComputeFramework/Private/ComputeKernelShaderCompilationManager.cpp b/Engine/Plugins/Runtime/ComputeFramework/Source/ComputeFramework/Private/ComputeKernelShaderCompilationManager.cpp
--- a/Engine/Plugins/Runtime/ComputeFramework/Source/ComputeFramework/Private/ComputeKernelShaderCompilationManager.cpp
+++ b/Engine/Plugins/Runtime/ComputeFramework/Source/ComputeFramework/Private/ComputeKernelShaderCompilationManager.cpp
@@ -29,10 +29,8 @@ FComputeKernelShaderCompilationManager GComputeKernelShaderCompilationManager;
 
 void FComputeKernelShaderCompilationManager::Tick(float DeltaSeconds)
 {
-#if WITH_EDITOR
 	RunCompileJobs();
 	ProcessAsyncResults();
-#endif
 }
 
 FComputeKernelShaderCompilationManager::FComputeKernelShaderCompilationManager()
@@ -52,7 +50,6 @@ FComputeKernelShaderCompilationManager::~FComputeKernelShaderCompilationManager(
 
 void FComputeKernelShaderCompilationManager::RunCompileJobs()
 {
-#if WITH_EDITOR
 	check(IsInGameThread());
 
 	InitWorkerInfo();
@@ -178,7 +175,6 @@ void FComputeKernelShaderCompilationManager::RunCompileJobs()
 		CurrentWorkerInfo.bComplete = false;
 		CurrentWorkerInfo.QueuedJobs.Empty();
 	}
-#endif
 }
 
 void FComputeKernelShaderCompilationManager::InitWorkerInfo()
@@ -198,7 +194,6 @@ void FComputeKernelShaderCompilationManager::InitWorkerInfo()
 
 void FComputeKernelShaderCompilationManager::AddJobs(TArray<FShaderCommonCompileJobPtr> InNewJobs)
 {
-#if WITH_EDITOR
 	for (auto& Job : InNewJobs)
 	{
 		FComputeKernelShaderMapCompileResults& ShaderMapInfo = ComputeKernelShaderMapJobs.FindOrAdd(Job->Id);
@@ -207,13 +202,11 @@ void FComputeKernelShaderCompilationManager::AddJobs(TArray<FShaderCommonCompile
 	}
 
 	JobQueue.Append(InNewJobs);
-#endif
 }
 
 
 void FComputeKernelShaderCompilationManager::ProcessAsyncResults()
 {
-#if WITH_EDITOR
 	int32 NumCompilingComputeKernelShaderMaps = 0;
 	TArray<int32> ShaderMapsToRemove;
 
@@ -241,7 +234,6 @@ void FComputeKernelShaderCompilationManager::ProcessAsyncResults()
 	{
 		ProcessCompiledComputeKernelShaderMaps(PendingFinalizeComputeKernelShaderMaps, 0.1f);
 	}
-#endif
 }
 
 
@@ -249,7 +241,6 @@ void FComputeKernelShaderCompilationManager::ProcessCompiledComputeKernelShaderM
 	TMap<int32, FComputeKernelShaderMapFinalizeResults>& CompiledShaderMaps,
 	float TimeBudget)
 {
-#if WITH_EDITOR
 	// Keeps shader maps alive as they are passed from the shader compiler and applied to the owning kernel
 	TArray<TRefCountPtr<FComputeKernelShaderMap> > LocalShaderMapReferences;
 	TMap<FComputeKernelResource*, FComputeKernelShaderMap*> KernelsToUpdate;
@@ -451,19 +442,17 @@ void FComputeKernelShaderCompilationManager::ProcessCompiledComputeKernelShaderM
 			}
 		}
 	}
-#endif
 }
 
 
 void FComputeKernelShaderCompilationManager::FinishCompilation(const TCHAR* InKernelName, const TArray<int32>& ShaderMapIdsToFinishCompiling)
 {
-#if WITH_EDITOR
 	check(!FPlatformProperties::RequiresCookedData());
 
 	RunCompileJobs();	// since we don't async compile through another process, this will run all oustanding jobs
 	ProcessAsyncResults();	// grab compiled shader maps and assign them to their resources
 
 	check(ComputeKernelShaderMapJobs.Num() == 0);
-#endif
 }
-#endif
+
+#endif // WITH_EDITOR
diff --git a/Engine/Plugins/Runtime/ComputeFramework/Source/ComputeFramework/Private/ComputeKernelShaderMap.cpp b/Engine/Plugins/Runtime/ComputeFramework/Source/ComputeFramework/Private/ComputeKernelShaderMap.cpp
--- a/Engine/Plugins/Runtime/ComputeFramework/Source/ComputeFramework/Private/ComputeKernelShaderMap.cpp
+++ b/Engine/Plugins/Runtime/ComputeFramework/Source/ComputeFramework/Private/ComputeKernelShaderMap.cpp
@@ -40,6 +40,8 @@ namespace ComputeKernelShaderCookStats
 TMap<FComputeKernelShaderMapId, FComputeKernelShaderMap*> FComputeKernelShaderMap::GIdToComputeKernelShaderMap[SP_NumPlatforms];
 TArray<FComputeKernelShaderMap*> FComputeKernelShaderMap::AllComputeKernelShaderMaps;
 
+#if WITH_EDITOR
+
 // The Id of 0 is reserved for global shaders
 uint32 FComputeKernelShaderMap::NextCompilingId = 2;
 
@@ -49,6 +51,7 @@ uint32 FComputeKernelShaderMap::NextCompilingId = 2;
  * Uses a TRefCountPtr as this will be the only reference to a shader map while it is being compiled.
  */
 TMap<TRefCountPtr<FComputeKernelShaderMap>, TArray<FComputeKernelResource*> > FComputeKernelShaderMap::ComputeKernelShaderMapsBeingCompiled;
+#endif // WITH_EDITOR
 
 
 
@@ -107,9 +110,9 @@ bool FComputeKernelShaderMapId::operator==(const FComputeKernelShaderMapId& InRe
 	return true;
 }
 
+#if WITH_EDITOR
 void FComputeKernelShaderMapId::AppendKeyString(FString& OutKeyString) const
 {
-#if WITH_EDITOR
 	OutKeyString.Appendf(TEXT("%llX"), ShaderCodeHash);
 	OutKeyString.AppendChar('_');
 
@@ -122,7 +125,6 @@ void FComputeKernelShaderMapId::AppendKeyString(FString& OutKeyString) const
 
 	// Add the inputs for any shaders that are stored inline in the shader map
 	AppendKeyStringShaderDependencies(MakeArrayView(ShaderTypeDependencies), LayoutParams, OutKeyString);
-#endif //WITH_EDITOR
 }
 
 /**
@@ -211,6 +213,7 @@ FShader* FComputeKernelShaderType::FinishCompileShader(
 
 	return Shader;
 }
+#endif // WITH_EDITOR
 
 /**
  * Finds the shader map for a kernel.
@@ -225,23 +228,20 @@ FComputeKernelShaderMap* FComputeKernelShaderMap::FindId(const FComputeKernelSha
 	return Result;
 }
 
+#if WITH_EDITOR
+
 /** Creates a string key for the derived data cache given a shader map id. */
 static FString GetComputeKernelShaderMapKeyString(const FComputeKernelShaderMapId& InShaderMapId, EShaderPlatform InPlatform)
 {
-#if WITH_EDITOR
 	const FName Format = LegacyShaderPlatformToShaderFormat(InPlatform);
 	FString ShaderMapKeyString = Format.ToString() + TEXT("_") + FString(FString::FromInt(GetTargetPlatformManagerRef().ShaderFormatVersion(Format))) + TEXT("_");
 	ShaderMapAppendKeyString(InPlatform, ShaderMapKeyString);
 	InShaderMapId.AppendKeyString(ShaderMapKeyString);
 	return FDerivedDataCacheInterface::BuildCacheKey(TEXT("CKERSM"), COMPUTEKERNEL_DERIVEDDATA_VER, *ShaderMapKeyString);
-#else
-	return FString();
-#endif
 }
 
 void FComputeKernelShaderMap::LoadFromDerivedDataCache(const FComputeKernelResource* InKernel, const FComputeKernelShaderMapId& InShaderMapId, EShaderPlatform InPlatform, TRefCountPtr<FComputeKernelShaderMap>& InOutShaderMap)
 {
-#if WITH_EDITOR
 	if (InOutShaderMap != nullptr)
 	{
 		check(InOutShaderMap->GetShaderPlatform() == InPlatform);
@@ -283,24 +283,21 @@ void FComputeKernelShaderMap::LoadFromDerivedDataCache(const FComputeKernelResou
 		}
 		INC_FLOAT_STAT_BY(STAT_ShaderCompiling_DDCLoading,(float)ComputeKernelShaderDDCTime);
 	}
-#endif
 }
 
 void FComputeKernelShaderMap::SaveToDerivedDataCache()
 {
-#if WITH_EDITOR
 	COOK_STAT(auto Timer = ComputeKernelShaderCookStats::UsageStats.TimeSyncWork());
-	
-TArray<uint8> SaveData;
+
+	TArray<uint8> SaveData;
 	FMemoryWriter Ar(SaveData, true);
 	Serialize(Ar);
 
 	const FString DataKey = GetComputeKernelShaderMapKeyString(GetContent()->ShaderMapId, GetShaderPlatform());
 
 	GetDerivedDataCacheRef().Put(*DataKey, SaveData, GetFriendlyName());
-	
-COOK_STAT(Timer.AddMiss(SaveData.Num()));
-#endif
+
+	COOK_STAT(Timer.AddMiss(SaveData.Num()));
 }
 
 /**
@@ -323,7 +320,6 @@ void FComputeKernelShaderMap::Compile(
 	{
 		UE_LOG(LogComputeFramework, Fatal, TEXT("Trying to compile ComputeKernel shader %s at run-time, which is not supported on consoles!"), *InKernel->GetFriendlyName() );
 	}
-#if WITH_EDITOR
 	else
 	{
 		// Make sure we are operating on a referenced shader map or the below Find will cause this shader map to be deleted,
@@ -437,7 +433,6 @@ void FComputeKernelShaderMap::Compile(
 			}
 		}
 	}
-#endif
 }
 
 FShader* FComputeKernelShaderMap::ProcessCompilationResultsForSingleJob(FShaderCompileJob& CurrentJob, const FSHAHash& InShaderMapHash)
@@ -512,6 +507,7 @@ bool FComputeKernelShaderMap::TryToAddToExistingCompilationTask(FComputeKernelRe
 
 	return false;
 }
+#endif // WITH_EDITOR
 
 bool FComputeKernelShaderMap::IsComputeKernelShaderComplete(const FComputeKernelResource* InKernel, const FComputeKernelShaderType* InShaderType, bool bSilent)
 {
@@ -539,6 +535,8 @@ bool FComputeKernelShaderMap::IsComplete(const FComputeKernelResource* InKernel,
 	// Make sure we are operating on a referenced shader map or the below Find will cause this shader map to be deleted,
 	// Since it creates a temporary ref counted pointer.
 	check(NumRefs > 0);
+
+#if WITH_EDITOR
 	const TArray<FComputeKernelResource*>* CorrespondingKernels = FComputeKernelShaderMap::ComputeKernelShaderMapsBeingCompiled.Find(this);
 
 	if (CorrespondingKernels)
@@ -546,6 +544,7 @@ bool FComputeKernelShaderMap::IsComplete(const FComputeKernelResource* InKernel,
 		check(!bCompilationFinalized);
 		return false;
 	}
+#endif
 
 	// Iterate over all shader types.
 	for(TLinkedList<FShaderType*>::TIterator ShaderTypeIt(FShaderType::GetTypeList());ShaderTypeIt;ShaderTypeIt.Next())
@@ -561,6 +560,7 @@ bool FComputeKernelShaderMap::IsComplete(const FComputeKernelResource* InKernel,
 	return true;
 }
 
+#if WITH_EDITOR
 void FComputeKernelShaderMap::LoadMissingShadersFromMemory(const FComputeKernelResource* InKernel)
 {
 #if 0
@@ -596,6 +596,8 @@ void FComputeKernelShaderMap::LoadMissingShadersFromMemory(const FComputeKernelR
 #endif
 }
 
+#endif // WITH_EDITOR
+
 void FComputeKernelShaderMap::GetShaderList(TMap<FShaderId, TShaderRef<FShader>>& OutShaders) const
 {
 	GetContent()->GetShaderList(*this, FSHAHash(), OutShaders);
@@ -661,23 +663,6 @@ FComputeKernelShaderMap::~FComputeKernelShaderMap()
 	AllComputeKernelShaderMaps.RemoveSwap(this);
 }
 
-/**
- * Removes all entries in the cache with exceptions based on a shader type
- * @param ShaderType - The shader type to flush
- */
-void FComputeKernelShaderMap::FlushShadersByShaderType(const FShaderType* InShaderType)
-{
-	if (InShaderType->GetComputeKernelShaderType())
-	{
-		const int32 PermutationCount = InShaderType->GetPermutationCount();
-		for (int32 PermutationId = 0; PermutationId < PermutationCount; ++PermutationId)
-		{
-			GetMutableContent()->RemoveShaderTypePermutaion(InShaderType->GetComputeKernelShaderType(), PermutationId);	
-		}
-	}
-}
-
-
 bool FComputeKernelShaderMap::Serialize(FArchive& Ar, bool bInlineShaderResources)
 {
 	// Note: This is saved to the DDC, not into packages (except when cooked)
@@ -686,6 +671,8 @@ bool FComputeKernelShaderMap::Serialize(FArchive& Ar, bool bInlineShaderResource
 	return Super::Serialize(Ar, bInlineShaderResources, false);
 }
 
+#if WITH_EDITOR
+
 void FComputeKernelShaderMap::RemovePending(FComputeKernelResource* InKernel)
 {
 	for (TMap<TRefCountPtr<FComputeKernelShaderMap>, TArray<FComputeKernelResource*> >::TIterator It(ComputeKernelShaderMapsBeingCompiled); It; ++It)
@@ -706,3 +693,4 @@ void FComputeKernelShaderMap::RemovePending(FComputeKernelResource* InKernel)
 	}
 }
 
+#endif // WITH_EDITOR
\ No newline at end of file
diff --git a/Engine/Plugins/Runtime/ComputeFramework/Source/ComputeFramework/Private/ComputeKernelShared.cpp b/Engine/Plugins/Runtime/ComputeFramework/Source/ComputeFramework/Private/ComputeKernelShared.cpp
--- a/Engine/Plugins/Runtime/ComputeFramework/Source/ComputeFramework/Private/ComputeKernelShared.cpp
+++ b/Engine/Plugins/Runtime/ComputeFramework/Source/ComputeFramework/Private/ComputeKernelShared.cpp
@@ -33,9 +33,12 @@ FComputeKernelResource::FComputeKernelResource()
 
 FComputeKernelResource::~FComputeKernelResource()
 {
+#if WITH_EDITOR
 	FComputeKernelShaderMap::RemovePending(this);
+#endif
 }
 
+#if WITH_EDITOR
 /** Populates OutEnvironment with defines needed to compile shaders for this kernel. */
 void FComputeKernelResource::SetupShaderCompilationEnvironment(EShaderPlatform InPlatform, FShaderCompilerEnvironment& OutEnvironment) const
 {
@@ -46,6 +49,7 @@ void FComputeKernelResource::SetupShaderCompilationEnvironment(EShaderPlatform I
 
 	OutEnvironment.IncludeVirtualPathToContentsMap.Append(AdditionalSources);
 }
+#endif // WITH_EDITOR
 
 bool FComputeKernelResource::ShouldCache(EShaderPlatform InPlatform, const FShaderType* InShaderType) const
 {
@@ -60,9 +64,9 @@ void FComputeKernelResource::NotifyCompilationFinished(FString const& ResultMess
 	OnCompilationCompleteDelegate.ExecuteIfBound(this);
 }
 
+#if WITH_EDITOR
 void FComputeKernelResource::CancelCompilation()
 {
-#if WITH_EDITOR
 	if (IsInGameThread())
 	{
 		FComputeKernelShaderMap::RemovePending(this);
@@ -71,8 +75,8 @@ void FComputeKernelResource::CancelCompilation()
 		FString Message = FString::Printf(TEXT("%s: Compilation cancelled."), *GetFriendlyName());
 		NotifyCompilationFinished(Message);
 	}
-#endif
 }
+#endif // WITH_EDITOR
 
 void FComputeKernelResource::RemoveOutstandingCompileId(const int32 InOldOutstandingCompileShaderMapId)
 {
@@ -84,7 +88,9 @@ void FComputeKernelResource::RemoveOutstandingCompileId(const int32 InOldOutstan
 
 void FComputeKernelResource::Invalidate()
 {
+#if WITH_EDITOR
 	CancelCompilation();
+#endif
 	ReleaseShaderMap();
 }
 
@@ -123,8 +129,8 @@ void FComputeKernelResource::GetShaderMapId(EShaderPlatform InPlatform, const IT
 
 		OutId.FeatureLevel = GetFeatureLevel();
 		OutId.ShaderCodeHash = ShaderCodeHash;
-		OutId.SetShaderDependencies(ShaderTypes, InPlatform);
 #if WITH_EDITOR
+		OutId.SetShaderDependencies(ShaderTypes, InPlatform);
 		if (TargetPlatform)
 		{
 			OutId.LayoutParams.InitializeForPlatform(TargetPlatform->IniPlatformName(), TargetPlatform->HasEditorOnlyData());
@@ -250,6 +256,7 @@ int32 FComputeKernelResource::GetNumPermutations() const
 	return ShaderPermutationVector.IsValid() ? 1 << ShaderPermutationVector->BitCount : 1;
 }
 
+#if WITH_EDITOR
 void FComputeKernelResource::SetupCompileEnvironment(int32 InPermutationId, FShaderCompilerEnvironment& OutShaderEnvironment) const
 {
  	for (TPair<FString, uint32> const& Permutation : ShaderPermutationVector->Permutations)
@@ -264,6 +271,7 @@ void FComputeKernelResource::SetupCompileEnvironment(int32 InPermutationId, FSha
 		}
 	}
 }
+#endif // WITH_EDITOR
 
 void FComputeKernelResource::SetRenderingThreadShaderMap(FComputeKernelShaderMap* InShaderMap)
 {
@@ -323,6 +331,7 @@ bool FComputeKernelResource::CacheShaders(const FComputeKernelShaderMapId& InSha
 			// Attempt to load from the derived data cache if we are uncooked
 			if ((!GameThreadShaderMap || !GameThreadShaderMap->IsComplete(this, true)) && !FPlatformProperties::RequiresCookedData())
 			{
+#if WITH_EDITOR
 				FComputeKernelShaderMap::LoadFromDerivedDataCache(this, InShaderMapId, InPlatform, GameThreadShaderMap);
 				if (GameThreadShaderMap && GameThreadShaderMap->IsValid())
 				{
@@ -333,6 +342,7 @@ bool FComputeKernelResource::CacheShaders(const FComputeKernelShaderMapId& InSha
 				{
 					UE_LOG(LogComputeFramework, Log, TEXT("Loading shader for kernel %s from DDC failed. Shader needs recompile."), *GetFriendlyName());
 				}
+#endif // WITH_EDITOR
 			}
 		}
 	}
@@ -342,6 +352,7 @@ bool FComputeKernelResource::CacheShaders(const FComputeKernelShaderMapId& InSha
 	bAssumeShaderMapIsComplete = (bContainsInlineShaders || FPlatformProperties::RequiresCookedData());
 #endif
 
+#if WITH_EDITOR
 	if (GameThreadShaderMap && GameThreadShaderMap->TryToAddToExistingCompilationTask(this))
 	{
 #if DEBUG_INFINITESHADERCOMPILE
@@ -353,7 +364,9 @@ bool FComputeKernelResource::CacheShaders(const FComputeKernelShaderMapId& InSha
 		GameThreadShaderMap = nullptr;
 		bSucceeded = true;
 	}
-	else if (!GameThreadShaderMap || !(bAssumeShaderMapIsComplete || GameThreadShaderMap->IsComplete(this, false)))
+	else
+#endif // WITH_EDITOR
+	if (!GameThreadShaderMap || !(bAssumeShaderMapIsComplete || GameThreadShaderMap->IsComplete(this, false)))
 	{
 		if (bContainsInlineShaders || FPlatformProperties::RequiresCookedData())
 		{
@@ -362,11 +375,13 @@ bool FComputeKernelResource::CacheShaders(const FComputeKernelShaderMapId& InSha
 		}
 		else
 		{
+#if WITH_EDITOR
 			UE_LOG(LogComputeFramework, Log, TEXT("%s cached shader map for kernel %s, compiling."), GameThreadShaderMap ? TEXT("Incomplete") : TEXT("Missing"), *GetFriendlyName());
 
 			// If there's no cached shader map for this kernel compile a new one.
 			// This is just kicking off the compile, GameThreadShaderMap will not be complete yet
 			bSucceeded = BeginCompileShaderMap(InShaderMapId, InPlatform, GameThreadShaderMap, bApplyCompletedShaderMapForRendering, bSynchronous);
+#endif
 
 			if (!bSucceeded)
 			{
@@ -394,9 +409,9 @@ bool FComputeKernelResource::CacheShaders(const FComputeKernelShaderMapId& InSha
 	return bSucceeded;
 }
 
+#if WITH_EDITOR
 void FComputeKernelResource::FinishCompilation()
 {
-#if WITH_EDITOR
 	TArray<int32> ShaderMapIdsToFinish;
 	GetShaderMapIDsWithUnfinishedCompilation(ShaderMapIdsToFinish);
 
@@ -415,8 +430,8 @@ void FComputeKernelResource::FinishCompilation()
 		GetShaderMapIDsWithUnfinishedCompilation(ShaderMapIdsToFinish2);
 		ensure(ShaderMapIdsToFinish2.Num() == 0);
 	}
-#endif
 }
+#endif // WITH_EDITOR
 
 TShaderRef<FComputeKernelShader> FComputeKernelResource::GetShader(int32 PermutationId) const
 {
@@ -426,7 +441,7 @@ TShaderRef<FComputeKernelShader> FComputeKernelResource::GetShader(int32 Permuta
 		return RenderingThreadShaderMap->GetShader<FComputeKernelShader>(PermutationId);
 	}
 	return TShaderRef<FComputeKernelShader>();
-};
+}
 
 TShaderRef<FComputeKernelShader> FComputeKernelResource::GetShaderGameThread(int32 PermutationId) const
 {
@@ -435,8 +450,9 @@ TShaderRef<FComputeKernelShader> FComputeKernelResource::GetShaderGameThread(int
 		return GameThreadShaderMap->GetShader<FComputeKernelShader>(PermutationId);
 	}
 	return TShaderRef<FComputeKernelShader>();
-};
+}
 
+#if WITH_EDITOR
 void FComputeKernelResource::GetShaderMapIDsWithUnfinishedCompilation(TArray<int32>& OutShaderMapIds)
 {
 	// Build an array of the shader map Id's are not finished compiling.
@@ -460,7 +476,6 @@ void FComputeKernelResource::GetShaderMapIDsWithUnfinishedCompilation(TArray<int
  */
 bool FComputeKernelResource::BeginCompileShaderMap(const FComputeKernelShaderMapId& InShaderMapId, EShaderPlatform InPlatform, TRefCountPtr<FComputeKernelShaderMap>& OutShaderMap, bool bApplyCompletedShaderMapForRendering, bool bSynchronous)
 {
-#if WITH_EDITORONLY_DATA
 	bool bSuccess = false;
 
 	STAT(double ComputeKernelCompileTime = 0);
@@ -491,15 +506,10 @@ bool FComputeKernelResource::BeginCompileShaderMap(const FComputeKernelShaderMap
 	}
 
 	return true;
-#else
-	UE_LOG(LogComputeFramework, Fatal, TEXT("Compiling of shaders in a build without editordata is not supported."));
-	return false;
-#endif
 }
 
 void FComputeKernelShaderMapId::SetShaderDependencies(const TArray<FShaderType*>& InShaderTypes, EShaderPlatform InShaderPlatform)
 {
-#if WITH_EDITOR
 	if (!FPlatformProperties::RequiresCookedData())
 	{
 		for (FShaderType* ShaderType : InShaderTypes)
@@ -513,9 +523,10 @@ void FComputeKernelShaderMapId::SetShaderDependencies(const TArray<FShaderType*>
 			}
 		}
 	}
-#endif //WITH_EDITOR
 }
 
+#endif // WITH_EDITOR
+
 bool FComputeKernelShaderMapId::ContainsShaderType(const FShaderType* ShaderType) const
 {
 	for (int32 TypeIndex = 0; TypeIndex < ShaderTypeDependencies.Num(); TypeIndex++)
