Fix include order in compute graph kernel building.
Need to have platform.ush before any custom struct declarations, since they use StructuredBuffer syntax which is #define'd from the different native syntax on some platforms.

diff --git a/Engine/Plugins/Animation/DeformerGraph/Source/OptimusCore/Private/Nodes/OptimusNode_ComputeKernelBase.cpp b/Engine/Plugins/Animation/DeformerGraph/Source/OptimusCore/Private/Nodes/OptimusNode_ComputeKernelBase.cpp
--- a/Engine/Plugins/Animation/DeformerGraph/Source/OptimusCore/Private/Nodes/OptimusNode_ComputeKernelBase.cpp
+++ b/Engine/Plugins/Animation/DeformerGraph/Source/OptimusCore/Private/Nodes/OptimusNode_ComputeKernelBase.cpp
@@ -179,9 +179,7 @@ FOptimus_ComputeKernelResult UOptimusNode_ComputeKernelBase::CreateComputeKernel
 		}
 	}
 
-	FString CookedSource =
-		"#include \"/Engine/Private/Common.ush\"\n"
-		"#include \"/Plugin/ComputeFramework/Private/ComputeKernelCommon.ush\"\n\n";
+	FString CookedSource;
 	CookedSource += FString::Join(GeneratedFunctions, TEXT("\n"));
 	CookedSource += "\n\n";
 	CookedSource += GetKernelSourceText();
diff --git a/Engine/Plugins/Runtime/ComputeFramework/Shaders/Private/ComputeKernel.usf b/Engine/Plugins/Runtime/ComputeFramework/Shaders/Private/ComputeKernel.usf
--- a/Engine/Plugins/Runtime/ComputeFramework/Shaders/Private/ComputeKernel.usf
+++ b/Engine/Plugins/Runtime/ComputeFramework/Shaders/Private/ComputeKernel.usf
@@ -1,7 +1,5 @@
 // Copyright Epic Games, Inc. All Rights Reserved.
 
-//todo[CF]: 
-// Remove includes to keep the compute framework sandboxed as much as possible
 #include "/Engine/Private/Common.ush"
 
 [numthreads(32, 1, 1)]
diff --git a/Engine/Plugins/Runtime/ComputeFramework/Shaders/Private/ComputeKernelCommon.ush b/Engine/Plugins/Runtime/ComputeFramework/Shaders/Private/ComputeKernelCommon.ush
--- a/Engine/Plugins/Runtime/ComputeFramework/Shaders/Private/ComputeKernelCommon.ush
+++ b/Engine/Plugins/Runtime/ComputeFramework/Shaders/Private/ComputeKernelCommon.ush
@@ -2,6 +2,8 @@
 
 #pragma once
 
+#include "/Engine/Private/Common.ush"
+
 // Macro to mark a kernel entry point.
 // Only one should exist per kernel.
 
diff --git a/Engine/Plugins/Runtime/ComputeFramework/Source/ComputeFramework/Private/ComputeGraph.cpp b/Engine/Plugins/Runtime/ComputeFramework/Source/ComputeFramework/Private/ComputeGraph.cpp
--- a/Engine/Plugins/Runtime/ComputeFramework/Source/ComputeFramework/Private/ComputeGraph.cpp
+++ b/Engine/Plugins/Runtime/ComputeFramework/Source/ComputeFramework/Private/ComputeGraph.cpp
@@ -469,9 +469,6 @@ FString UComputeGraph::BuildKernelSource(
 	TSet<FString> StructsSeen;
 	TArray<FString> StructDeclarations;
 
-	// Add include for any platform specific types.
-	HLSL += TEXT("#include \"/Engine/Public/Platform.ush\"\n");
-
 	// Add virtual source includes from the additional sources.
 	for (TPair<FString, FString> const& AdditionalSource : InAdditionalSources)
 	{
@@ -563,6 +560,11 @@ FString UComputeGraph::BuildKernelSource(
 	HLSL += InKernelSource.GetSource();
 
 	FString Declaration;
+
+	// Add include at top of code. This is required to get platform.ush early enough to provide platform specific types.
+	Declaration += TEXT("#include \"/Plugin/ComputeFramework/Private/ComputeKernelCommon.ush\"\n\n");
+	GetShaderFileHash(TEXT("/Plugin/ComputeFramework/Private/ComputeKernelCommon.ush"), EShaderPlatform::SP_PCD3D_SM5).AppendString(OutHashKey);
+
 	for (const FString& StructDeclaration : StructDeclarations)
 	{
 		Declaration += StructDeclaration;
