Remove ShaderPlatformToShaderFormatName

diff --git a/Engine/Plugins/Runtime/ComputeFramework/Source/ComputeFramework/Private/ComputeKernelShaderCompilationManager.cpp b/Engine/Plugins/Runtime/ComputeFramework/Source/ComputeFramework/Private/ComputeKernelShaderCompilationManager.cpp
--- a/Engine/Plugins/Runtime/ComputeFramework/Source/ComputeFramework/Private/ComputeKernelShaderCompilationManager.cpp
+++ b/Engine/Plugins/Runtime/ComputeFramework/Source/ComputeFramework/Private/ComputeKernelShaderCompilationManager.cpp
@@ -386,15 +386,15 @@ void FComputeKernelShaderCompilationManager::ProcessCompiledComputeKernelShaderM
 					{
 						if (CompletedShaderMap->IsComplete(Kernel, true))
 						{
-							const FString PlatformName = ShaderPlatformToShaderFormatName(ShaderMap->GetShaderPlatform()).ToString();
+							const FString ShaderFormatName = FDataDrivenShaderPlatformInfo::GetShaderFormat(ShaderMap->GetShaderPlatform()).ToString();
 							if (bSuccess)
 							{
-								const FString SuccessMessage = FString::Printf(TEXT("%s: %s shader compilation success!"), *Kernel->GetFriendlyName(), *PlatformName);
+								const FString SuccessMessage = FString::Printf(TEXT("%s: %s shader compilation success!"), *Kernel->GetFriendlyName(), *ShaderFormatName);
 								Kernel->NotifyCompilationFinished(SuccessMessage);
 							}
 							else
 							{
-								const FString FailMessage = FString::Printf(TEXT("%s: %s shader compilation failed."), *Kernel->GetFriendlyName(), *PlatformName);
+								const FString FailMessage = FString::Printf(TEXT("%s: %s shader compilation failed."), *Kernel->GetFriendlyName(), *ShaderFormatName);
 								Kernel->NotifyCompilationFinished(FailMessage);
 							}
 						}
@@ -430,9 +430,10 @@ void FComputeKernelShaderCompilationManager::ProcessCompiledComputeKernelShaderM
 
 			if (ShaderMap && ShaderMap->CompiledSuccessfully())
 			{
+				const FString ShaderFormatName = FDataDrivenShaderPlatformInfo::GetShaderFormat(ShaderMap->GetShaderPlatform()).ToString();
 				const FString SuccessMessage = FString::Printf(TEXT("%s: %s shader compilation success!"), 
 					*Kernel->GetFriendlyName(), 
-					*ShaderPlatformToShaderFormatName(ShaderMap->GetShaderPlatform()).ToString());
+					*ShaderFormatName);
 				Kernel->NotifyCompilationFinished(SuccessMessage);
 			}
 			else
diff --git a/Engine/Plugins/Runtime/ComputeFramework/Source/ComputeFramework/Private/ComputeKernelShaderMap.cpp b/Engine/Plugins/Runtime/ComputeFramework/Source/ComputeFramework/Private/ComputeKernelShaderMap.cpp
--- a/Engine/Plugins/Runtime/ComputeFramework/Source/ComputeFramework/Private/ComputeKernelShaderMap.cpp
+++ b/Engine/Plugins/Runtime/ComputeFramework/Source/ComputeFramework/Private/ComputeKernelShaderMap.cpp
@@ -390,10 +390,11 @@ void FComputeKernelShaderMap::Compile(
 				else if (ShaderType)
 				{
 					InKernel->RemoveOutstandingCompileId(CompilingId);
+					const FString ShaderFormatName = FDataDrivenShaderPlatformInfo::GetShaderFormat(InPlatform).ToString();
 
 					FString Message = FString::Printf(TEXT("%s: Compilation not supported on %s."), 
 						*InKernel->GetFriendlyName(), 
-						*ShaderPlatformToShaderFormatName(InPlatform).ToString());
+						*ShaderFormatName);
 					InKernel->NotifyCompilationFinished(Message);
 				}
 			}
diff --git a/Engine/Plugins/Runtime/ComputeFramework/Source/ComputeFramework/Private/ComputeKernelShared.cpp b/Engine/Plugins/Runtime/ComputeFramework/Source/ComputeFramework/Private/ComputeKernelShared.cpp
--- a/Engine/Plugins/Runtime/ComputeFramework/Source/ComputeFramework/Private/ComputeKernelShared.cpp
+++ b/Engine/Plugins/Runtime/ComputeFramework/Source/ComputeFramework/Private/ComputeKernelShared.cpp
@@ -391,7 +391,8 @@ bool FComputeKernelResource::CacheShaders(const FComputeKernelShaderMapId& InSha
 	}
 	else
 	{
-		NotifyCompilationFinished(FString::Printf(TEXT("%s: Loaded %s shaders from DDC."), *GetFriendlyName(), *ShaderPlatformToShaderFormatName(InPlatform).ToString()));
+		const FString ShaderFormatName = FDataDrivenShaderPlatformInfo::GetShaderFormat(InPlatform).ToString();
+		NotifyCompilationFinished(FString::Printf(TEXT("%s: Loaded %s shaders from DDC."), *GetFriendlyName(), *ShaderFormatName));
 		bSucceeded = true;
 	}
 
diff --git a/Engine/Source/Developer/ShaderPreprocessor/Private/ShaderPreprocessor.cpp b/Engine/Source/Developer/ShaderPreprocessor/Private/ShaderPreprocessor.cpp
--- a/Engine/Source/Developer/ShaderPreprocessor/Private/ShaderPreprocessor.cpp
+++ b/Engine/Source/Developer/ShaderPreprocessor/Private/ShaderPreprocessor.cpp
@@ -93,7 +93,7 @@ private:
 			}
 			else
 			{
-				CheckShaderHashCacheInclude(VirtualFilePath, This->ShaderInput.Target.GetPlatform());
+				CheckShaderHashCacheInclude(VirtualFilePath, This->ShaderInput.Target.GetPlatform(), This->ShaderInput.ShaderFormat.ToString());
 
 				LoadShaderSourceFile(*VirtualFilePath, This->ShaderInput.Target.GetPlatform(), &FileContents, &This->ShaderOutput.Errors, &This->ShaderInput.ShaderPlatformName);
 			}
diff --git a/Engine/Source/Runtime/RHI/Public/RHIShaderFormatDefinitions.inl b/Engine/Source/Runtime/RHI/Public/RHIShaderFormatDefinitions.inl
--- a/Engine/Source/Runtime/RHI/Public/RHIShaderFormatDefinitions.inl
+++ b/Engine/Source/Runtime/RHI/Public/RHIShaderFormatDefinitions.inl
@@ -30,60 +30,6 @@ static FName NAME_VULKAN_SM5_ANDROID(TEXT("SF_VULKAN_SM5_ANDROID"));
 
 static FName NAME_D3D_ES3_1_HOLOLENS(TEXT("D3D_ES3_1_HOLOLENS"));
 
-static FName ShaderPlatformToShaderFormatName(EShaderPlatform Platform)
-{
-	switch (Platform)
-	{
-	case SP_PCD3D_SM6:
-		return NAME_PCD3D_SM6;
-	case SP_PCD3D_SM5:
-		return NAME_PCD3D_SM5;
-	case SP_PCD3D_ES3_1:
-		return NAME_PCD3D_ES3_1;
-
-	case SP_OPENGL_PCES3_1:
-		return NAME_GLSL_150_ES31;
-	case SP_OPENGL_ES3_1_ANDROID:
-		return NAME_GLSL_ES3_1_ANDROID;
-
-	case SP_METAL:
-		return NAME_SF_METAL;
-	case SP_METAL_MRT:
-		return NAME_SF_METAL_MRT;
-	case SP_METAL_TVOS:
-		return NAME_SF_METAL_TVOS;
-	case SP_METAL_MRT_TVOS:
-		return NAME_SF_METAL_MRT_TVOS;
-	case SP_METAL_MRT_MAC:
-		return NAME_SF_METAL_MRT_MAC;
-	case SP_METAL_SM5:
-		return NAME_SF_METAL_SM5;
-	case SP_METAL_MACES3_1:
-		return NAME_SF_METAL_MACES3_1;
-	case SP_VULKAN_ES3_1_ANDROID:
-		return NAME_VULKAN_ES3_1_ANDROID;
-	case SP_VULKAN_PCES3_1:
-		return NAME_VULKAN_ES3_1;
-	case SP_VULKAN_SM5:
-		return NAME_VULKAN_SM5;
-	case SP_VULKAN_SM5_ANDROID:
-		return NAME_VULKAN_SM5_ANDROID;
-	case SP_D3D_ES3_1_HOLOLENS:
-		return NAME_D3D_ES3_1_HOLOLENS;
-
-	default:
-		if (FStaticShaderPlatformNames::IsStaticPlatform(Platform))
-		{
-			return FStaticShaderPlatformNames::Get().GetShaderFormat(Platform);
-		}
-		else
-		{
-			checkf(0, TEXT("Unknown EShaderPlatform %d!"), (int32)Platform);
-			return NAME_PCD3D_SM5;
-		}
-	}
-}
-
 static EShaderPlatform ShaderFormatNameToShaderPlatform(FName ShaderFormat)
 {
 	if (ShaderFormat == NAME_PCD3D_SM6)					return SP_PCD3D_SM6;
diff --git a/Engine/Source/Runtime/RenderCore/Private/ShaderCore.cpp b/Engine/Source/Runtime/RenderCore/Private/ShaderCore.cpp
--- a/Engine/Source/Runtime/RenderCore/Private/ShaderCore.cpp
+++ b/Engine/Source/Runtime/RenderCore/Private/ShaderCore.cpp
@@ -1481,7 +1481,7 @@ void UpdateIncludeDirectoryForPreviewPlatform(EShaderPlatform PreviewPlatform, E
 	GShaderHashCache.UpdateIncludeDirectoryForPreviewPlatform(PreviewPlatform, ActualPlatform);
 }
 
-void CheckShaderHashCacheInclude(const FString& VirtualFilePath, EShaderPlatform ShaderPlatform)
+void CheckShaderHashCacheInclude(const FString& VirtualFilePath, EShaderPlatform ShaderPlatform, const FString& ShaderFormatName)
 {
 	FRWScopeLock ShaderHashAccessLock(GShaderHashAccessRWLock, SLT_ReadOnly);
 	bool bIgnoreInclude = GShaderHashCache.ShouldIgnoreInclude(VirtualFilePath, ShaderPlatform);
@@ -1489,7 +1489,7 @@ void CheckShaderHashCacheInclude(const FString& VirtualFilePath, EShaderPlatform
 	checkf(!bIgnoreInclude,
 		TEXT("Shader compiler is trying to include %s, which is not located in IShaderFormat::GetPlatformIncludeDirectory for %s."),
 		*VirtualFilePath,
-		*ShaderPlatformToShaderFormatName(ShaderPlatform).ToString());
+		*ShaderFormatName);
 }
 
 void InitializeShaderTypes()
diff --git a/Engine/Source/Runtime/RenderCore/Public/ShaderCore.h b/Engine/Source/Runtime/RenderCore/Public/ShaderCore.h
--- a/Engine/Source/Runtime/RenderCore/Public/ShaderCore.h
+++ b/Engine/Source/Runtime/RenderCore/Public/ShaderCore.h
@@ -117,7 +117,7 @@ extern RENDERCORE_API void InitializeShaderHashCache();
 extern RENDERCORE_API void UpdateIncludeDirectoryForPreviewPlatform(EShaderPlatform PreviewPlatform, EShaderPlatform ActualPlatform);
 
 /** Checks if shader include isn't skipped by a shader hash cache. */
-extern RENDERCORE_API void CheckShaderHashCacheInclude(const FString& VirtualFilePath, EShaderPlatform ShaderPlatform);
+extern RENDERCORE_API void CheckShaderHashCacheInclude(const FString& VirtualFilePath, EShaderPlatform ShaderPlatform, const FString& ShaderFormatName);
 
 /** Initializes cached shader type data.  This must be called before creating any FShaderType. */
 extern RENDERCORE_API void InitializeShaderTypes();
