ComputeFramework: When testing if plugin is enabled test for the module actually being loaded.

diff --git a/Engine/Plugins/Runtime/ComputeFramework/Source/ComputeFramework/Private/ComputeFramework.cpp b/Engine/Plugins/Runtime/ComputeFramework/Source/ComputeFramework/Private/ComputeFramework.cpp
--- a/Engine/Plugins/Runtime/ComputeFramework/Source/ComputeFramework/Private/ComputeFramework.cpp
+++ b/Engine/Plugins/Runtime/ComputeFramework/Source/ComputeFramework/Private/ComputeFramework.cpp
@@ -38,7 +38,7 @@ namespace ComputeFramework
 
 	bool IsEnabled()
 	{
-		return GComputeFrameworkEnable > 0;
+		return GComputeFrameworkEnable > 0 && FComputeFrameworkModule::GetComputeSystem() != nullptr;
 	}
 
 	void RebuildComputeGraphs()
@@ -66,7 +66,8 @@ namespace ComputeFramework
 
 	void FlushWork(FSceneInterface const* InScene, FName InExecutionGroupName)
 	{
-		FComputeGraphTaskWorker* ComputeGraphWorker = FComputeFrameworkModule::GetComputeSystem()->GetComputeWorker(InScene);
+		FComputeFrameworkSystem* ComputeSystem = FComputeFrameworkModule::GetComputeSystem();
+		FComputeGraphTaskWorker* ComputeGraphWorker = ComputeSystem != nullptr ? ComputeSystem->GetComputeWorker(InScene) : nullptr;
 		if (ensure(ComputeGraphWorker))
 		{
 			ENQUEUE_RENDER_COMMAND(ComputeFrameworkFlushCommand)(
