ComputeFramework: Fix android build error.

diff --git a/Engine/Plugins/Runtime/ComputeFramework/Source/ComputeFramework/Private/ComputeGraphWorker.cpp b/Engine/Plugins/Runtime/ComputeFramework/Source/ComputeFramework/Private/ComputeGraphWorker.cpp
--- a/Engine/Plugins/Runtime/ComputeFramework/Source/ComputeFramework/Private/ComputeGraphWorker.cpp
+++ b/Engine/Plugins/Runtime/ComputeFramework/Source/ComputeFramework/Private/ComputeGraphWorker.cpp
@@ -69,7 +69,7 @@ void FComputeGraphTaskWorker::SubmitWork(FRDGBuilder& GraphBuilder, FName InExec
 						FComputeDataProviderRenderProxy* DataProvider = GraphInvocation.DataProviderRenderProxies[DataProviderIndex];
 						if (ensure(DataProvider != nullptr))
 						{
-							FComputeDataProviderRenderProxy::FValidationData ValidationData{ NumSubInvocations, Member.GetStructMetadata()->GetSize() };
+							FComputeDataProviderRenderProxy::FValidationData ValidationData{ NumSubInvocations, (int32)Member.GetStructMetadata()->GetSize() };
 							bIsValid &= DataProvider->IsValid(ValidationData);
 
 							if (bIsValid)
