ComputeFramework: We no longer need to check full shadermap for each graph kernel before enqueuing work.

diff --git a/Engine/Plugins/Runtime/ComputeFramework/Source/ComputeFramework/Private/ComputeGraph.cpp b/Engine/Plugins/Runtime/ComputeFramework/Source/ComputeFramework/Private/ComputeGraph.cpp
--- a/Engine/Plugins/Runtime/ComputeFramework/Source/ComputeFramework/Private/ComputeGraph.cpp
+++ b/Engine/Plugins/Runtime/ComputeFramework/Source/ComputeFramework/Private/ComputeGraph.cpp
@@ -121,30 +121,6 @@ bool UComputeGraph::ValidateGraph(FString* OutErrors)
 	return true;
 }
 
-bool UComputeGraph::IsCompiled() const
-{
-	// todo[CF]: Checking all shader maps is probably slow. Cache and serialize compilation success after each compile instead.
-	for (int32 KernelIndex = 0; KernelIndex < KernelInvocations.Num(); ++KernelIndex)
-	{
-		if (KernelInvocations[KernelIndex] != nullptr)
-		{
-			if (!KernelResources.IsValidIndex(KernelIndex))
-			{
-				return false;
-			}
-			
-			FComputeKernelResource const* Resource = KernelResources[KernelIndex].Get();
-			FComputeKernelShaderMap* ShaderMap = Resource != nullptr ? Resource->GetGameThreadShaderMap() : nullptr;
-			if (ShaderMap == nullptr || !ShaderMap->IsComplete(Resource, true))
-			{
-				return false;
-			}
-		}
-	}
-
-	return true;
-}
-
 bool UComputeGraph::ValidateProviders(TArray< TObjectPtr<UComputeDataProvider> > const& DataProviders) const
 {
 	if (DataInterfaces.Num() != DataProviders.Num())
diff --git a/Engine/Plugins/Runtime/ComputeFramework/Source/ComputeFramework/Private/ComputeGraphInstance.cpp b/Engine/Plugins/Runtime/ComputeFramework/Source/ComputeFramework/Private/ComputeGraphInstance.cpp
--- a/Engine/Plugins/Runtime/ComputeFramework/Source/ComputeFramework/Private/ComputeGraphInstance.cpp
+++ b/Engine/Plugins/Runtime/ComputeFramework/Source/ComputeFramework/Private/ComputeGraphInstance.cpp
@@ -26,11 +26,6 @@ void FComputeGraphInstance::DestroyDataProviders()
 	DataProviders.Reset();
 }
 
-bool FComputeGraphInstance::ValidateDataProviders(UComputeGraph* InComputeGraph) const
-{
-	return InComputeGraph != nullptr && InComputeGraph->IsCompiled() && InComputeGraph->ValidateGraph() && InComputeGraph->ValidateProviders(DataProviders);
-}
-
 bool FComputeGraphInstance::EnqueueWork(UComputeGraph* InComputeGraph, FSceneInterface const* InScene, FName InExecutionGroupName, FName InOwnerName, FSimpleDelegate InFallbackDelegate)
 {
 	if (InComputeGraph == nullptr || InScene == nullptr)
diff --git a/Engine/Plugins/Runtime/ComputeFramework/Source/ComputeFramework/Public/ComputeFramework/ComputeGraph.h b/Engine/Plugins/Runtime/ComputeFramework/Source/ComputeFramework/Public/ComputeFramework/ComputeGraph.h
--- a/Engine/Plugins/Runtime/ComputeFramework/Source/ComputeFramework/Public/ComputeFramework/ComputeGraph.h
+++ b/Engine/Plugins/Runtime/ComputeFramework/Source/ComputeFramework/Public/ComputeFramework/ComputeGraph.h
@@ -110,9 +110,6 @@ public:
 	 */
 	bool ValidateGraph(FString* OutErrors = nullptr);
 
-	/** Returns true if graph has a full set of compiled shaders. */
-	bool IsCompiled() const;
-
 	/** 
 	 * Create UComputeDataProvider objects to match the current UComputeDataInterface objects. 
 	 * We attempt to setup bindings from the InBindingObjects.
diff --git a/Engine/Plugins/Runtime/ComputeFramework/Source/ComputeFramework/Public/ComputeFramework/ComputeGraphInstance.h b/Engine/Plugins/Runtime/ComputeFramework/Source/ComputeFramework/Public/ComputeFramework/ComputeGraphInstance.h
--- a/Engine/Plugins/Runtime/ComputeFramework/Source/ComputeFramework/Public/ComputeFramework/ComputeGraphInstance.h
+++ b/Engine/Plugins/Runtime/ComputeFramework/Source/ComputeFramework/Public/ComputeFramework/ComputeGraphInstance.h
@@ -28,9 +28,6 @@ public:
 	/** Create the Data Provider objects. */
 	void DestroyDataProviders();
 
-	/** Returns true if the Data Provider objects are all created and valid. */
-	bool ValidateDataProviders(UComputeGraph* InComputeGraph) const;
-
 	/** Get the Data Provider objects. */
 	TArray< TObjectPtr<UComputeDataProvider> >& GetDataProviders() { return DataProviders; }
 
