ComputeFramework: Use generic on-disk shader file as source for compute framework kernels.
Source is included from there in generically named generated file.

diff --git a/Engine/Plugins/Runtime/ComputeFramework/Shaders/Private/ComputeKernel.usf b/Engine/Plugins/Runtime/ComputeFramework/Shaders/Private/ComputeKernel.usf
--- a/Engine/Plugins/Runtime/ComputeFramework/Shaders/Private/ComputeKernel.usf
+++ b/Engine/Plugins/Runtime/ComputeFramework/Shaders/Private/ComputeKernel.usf
@@ -1,8 +1,6 @@
 // Copyright Epic Games, Inc. All Rights Reserved.
 
-#include "/Engine/Private/Common.ush"
+#include "/Plugin/ComputeFramework/Private/ComputeKernelCommon.ush"
 
-[numthreads(32, 1, 1)]
-void ComputeKernelNull()
-{
-}
+// Include the generated kernel source.
+#include "/Engine/Generated/ComputeKernel.ush"
diff --git a/Engine/Plugins/Runtime/ComputeFramework/Source/ComputeFramework/Private/ComputeGraph.cpp b/Engine/Plugins/Runtime/ComputeFramework/Source/ComputeFramework/Private/ComputeGraph.cpp
--- a/Engine/Plugins/Runtime/ComputeFramework/Source/ComputeFramework/Private/ComputeGraph.cpp
+++ b/Engine/Plugins/Runtime/ComputeFramework/Source/ComputeFramework/Private/ComputeGraph.cpp
@@ -536,11 +536,6 @@ FString UComputeGraph::BuildKernelSource(
 	HLSL += InKernelSource.GetSource();
 
 	FString Declaration;
-
-	// Add include at top of code. This is required to get platform.ush early enough to provide platform specific types.
-	Declaration += TEXT("#include \"/Plugin/ComputeFramework/Private/ComputeKernelCommon.ush\"\n\n");
-	GetShaderFileHash(TEXT("/Plugin/ComputeFramework/Private/ComputeKernelCommon.ush"), EShaderPlatform::SP_PCD3D_SM5).AppendString(OutHashKey);
-
 	for (const FString& StructDeclaration : StructDeclarations)
 	{
 		Declaration += StructDeclaration;
@@ -553,6 +548,9 @@ FString UComputeGraph::BuildKernelSource(
 	// Finalize hash state and add to the final hash key.
 	HashState.Finalize().AppendString(OutHashKey);
 
+	// Add the our boilerplate wrapper to the final hash key.
+	GetShaderFileHash(TEXT("/Plugin/ComputeFramework/Private/ComputeKernel.usf"), EShaderPlatform::SP_PCD3D_SM5).AppendString(OutHashKey);
+
 	return HLSL;
 }
 
diff --git a/Engine/Plugins/Runtime/ComputeFramework/Source/ComputeFramework/Private/ComputeKernelShaderMap.cpp b/Engine/Plugins/Runtime/ComputeFramework/Source/ComputeFramework/Private/ComputeKernelShaderMap.cpp
--- a/Engine/Plugins/Runtime/ComputeFramework/Source/ComputeFramework/Private/ComputeKernelShaderMap.cpp
+++ b/Engine/Plugins/Runtime/ComputeFramework/Source/ComputeFramework/Private/ComputeKernelShaderMap.cpp
@@ -148,15 +148,16 @@ void FComputeKernelShaderType::BeginCompileShader(
 		EShaderCompileJobPriority::Normal
 		);
 
-	FString virtualSourcePath = FString::Printf(TEXT("/Engine/Generated/ComputeFramework/Kernel_%s.usf"), *InKernel->GetEntryPoint());
+	TCHAR const* SourceFilePath = TEXT("/Plugin/ComputeFramework/Private/ComputeKernel.usf");
+	TCHAR const* GeneratedSourceFilePath = TEXT("/Engine/Generated/ComputeKernel.ush");
 
 	NewJob->ShaderParameters = MakeShared<const FParameters, ESPMode::ThreadSafe>(*InKernel->GetShaderParamMetadata());
 	NewJob->Input.SharedEnvironment = InCompilationEnvironment;
 	NewJob->Input.Target = InTarget;
 	NewJob->Input.ShaderFormat = LegacyShaderPlatformToShaderFormat(InPlatform);
-	NewJob->Input.VirtualSourceFilePath = virtualSourcePath;
+	NewJob->Input.VirtualSourceFilePath = SourceFilePath;
 	NewJob->Input.EntryPointName = InKernel->GetEntryPoint();
-	NewJob->Input.Environment.IncludeVirtualPathToContentsMap.Add(virtualSourcePath, InKernel->GetHLSLSource());
+	NewJob->Input.Environment.IncludeVirtualPathToContentsMap.Add(GeneratedSourceFilePath, InKernel->GetHLSLSource());
 	UE_LOG(LogComputeFramework, Verbose, TEXT("%s"), *InKernel->GetHLSLSource());
 	
 	AddUniformBufferIncludesToEnvironment(NewJob->Input.Environment, InPlatform);
@@ -177,7 +178,7 @@ void FComputeKernelShaderType::BeginCompileShader(
 		this,
 		nullptr,//ShaderPipeline,
 		PermutationId,
-		*virtualSourcePath,
+		SourceFilePath,
 		*InKernel->GetEntryPoint(),
 		FShaderTarget(GetFrequency(), InPlatform),
 		NewJob->Input
