ComputeFramework: Virtual path changes for compute kernel and source library includes.
Move them to ComputeFramework virtual path by default and use valid paths with .ush suffix.
This is in preparation for migrating to the core shared compliler manager which validates these things.
Compile error path->object resolution for source libraries will be broken, but fixed in a second pass.

diff --git a/Engine/Plugins/Runtime/ComputeFramework/Shaders/Private/ComputeKernel.usf b/Engine/Plugins/Runtime/ComputeFramework/Shaders/Private/ComputeKernel.usf
--- a/Engine/Plugins/Runtime/ComputeFramework/Shaders/Private/ComputeKernel.usf
+++ b/Engine/Plugins/Runtime/ComputeFramework/Shaders/Private/ComputeKernel.usf
@@ -3,4 +3,4 @@
 #include "/Plugin/ComputeFramework/Private/ComputeKernelCommon.ush"
 
 // Include the generated kernel source.
-#include "/Engine/Generated/ComputeKernel.ush"
+#include "/Plugin/ComputeFramework/Generated/ComputeKernel.ush"
diff --git a/Engine/Plugins/Runtime/ComputeFramework/Source/ComputeFramework/Private/ComputeKernelShaderMap.cpp b/Engine/Plugins/Runtime/ComputeFramework/Source/ComputeFramework/Private/ComputeKernelShaderMap.cpp
--- a/Engine/Plugins/Runtime/ComputeFramework/Source/ComputeFramework/Private/ComputeKernelShaderMap.cpp
+++ b/Engine/Plugins/Runtime/ComputeFramework/Source/ComputeFramework/Private/ComputeKernelShaderMap.cpp
@@ -148,8 +148,10 @@ void FComputeKernelShaderType::BeginCompileShader(
 		EShaderCompileJobPriority::Normal
 		);
 
+	FString GeneratedSourceFilePath = FString::Printf(TEXT("/Plugin/ComputeFramework/Generated/Kernel_%s.ush"), *InKernel->GetEntryPoint());
+	FString GeneratedStubSource = FString::Printf(TEXT("#include \"%s\"\n"), *GeneratedSourceFilePath);
+	TCHAR const* GeneratedStubFilePath = TEXT("/Plugin/ComputeFramework/Generated/ComputeKernel.ush");
 	TCHAR const* SourceFilePath = TEXT("/Plugin/ComputeFramework/Private/ComputeKernel.usf");
-	TCHAR const* GeneratedSourceFilePath = TEXT("/Engine/Generated/ComputeKernel.ush");
 
 	NewJob->ShaderParameters = MakeShared<const FParameters, ESPMode::ThreadSafe>(*InKernel->GetShaderParamMetadata());
 	NewJob->Input.SharedEnvironment = InCompilationEnvironment;
@@ -158,6 +160,7 @@ void FComputeKernelShaderType::BeginCompileShader(
 	NewJob->Input.VirtualSourceFilePath = SourceFilePath;
 	NewJob->Input.EntryPointName = InKernel->GetEntryPoint();
 	NewJob->Input.Environment.IncludeVirtualPathToContentsMap.Add(GeneratedSourceFilePath, InKernel->GetHLSLSource());
+	NewJob->Input.Environment.IncludeVirtualPathToContentsMap.Add(GeneratedStubFilePath, GeneratedStubSource);
 	UE_LOG(LogComputeFramework, Verbose, TEXT("%s"), *InKernel->GetHLSLSource());
 	
 	AddUniformBufferIncludesToEnvironment(NewJob->Input.Environment, InPlatform);
diff --git a/Engine/Plugins/Runtime/ComputeFramework/Source/ComputeFramework/Public/ComputeFramework/ComputeSource.h b/Engine/Plugins/Runtime/ComputeFramework/Source/ComputeFramework/Public/ComputeFramework/ComputeSource.h
--- a/Engine/Plugins/Runtime/ComputeFramework/Source/ComputeFramework/Public/ComputeFramework/ComputeSource.h
+++ b/Engine/Plugins/Runtime/ComputeFramework/Source/ComputeFramework/Public/ComputeFramework/ComputeSource.h
@@ -18,7 +18,7 @@ public:
 	/** Get source code ready for HLSL compilation. */
 	virtual FString GetSource() const { return FString(); }
 	/** Get virtual file path for the source code. This will be the file name shown in any compilation errors. */
-	virtual FString GetVirtualPath() const { return GetPathName(); }
+	virtual FString GetVirtualPath() const { return FString::Printf(TEXT("/Plugin/ComputeFramework/Generated%s.ush"), *GetPathName()); }
 
 public:
 	/** Array of additional source objects. This allows us to specify source dependencies. */
