ComputeFramework: Move management of shader compilation to the core shader manager.
Existing system was single threaded. With this change we can make used of shader workers.

diff --git a/Engine/Plugins/Runtime/ComputeFramework/Source/ComputeFramework/Private/ComputeFramework/ComputeKernelShaderCompilationManager.h b/Engine/Plugins/Runtime/ComputeFramework/Source/ComputeFramework/Private/ComputeFramework/ComputeKernelShaderCompilationManager.h
--- a/Engine/Plugins/Runtime/ComputeFramework/Source/ComputeFramework/Private/ComputeFramework/ComputeKernelShaderCompilationManager.h
+++ b/Engine/Plugins/Runtime/ComputeFramework/Source/ComputeFramework/Private/ComputeFramework/ComputeKernelShaderCompilationManager.h
@@ -5,46 +5,6 @@
 #include "ComputeKernelShared.h"
 #include "ShaderCompiler.h"
 
-/** Information tracked for each shader compile worker process instance. */
-struct FComputeKernelShaderCompileWorkerInfo
-{
-	/** Process handle of the worker app once launched.  Invalid handle means no process. */
-	FProcHandle WorkerProcess;
-
-	/** Tracks whether tasks have been issued to the worker. */
-	bool bIssuedTasksToWorker;
-
-	/** Whether the worker has been launched for this set of tasks. */
-	bool bLaunchedWorker;
-
-	/** Tracks whether all tasks issued to the worker have been received. */
-	bool bComplete;
-
-	/** Time at which the worker started the most recent batch of tasks. */
-	double StartTime;
-
-	/** Jobs that this worker is responsible for compiling. */
-	TArray<FShaderCommonCompileJobPtr> QueuedJobs;
-
-	FComputeKernelShaderCompileWorkerInfo() :
-		bIssuedTasksToWorker(false),
-		bLaunchedWorker(false),
-		bComplete(false),
-		StartTime(0)
-	{
-	}
-
-	// warning: not virtual
-	~FComputeKernelShaderCompileWorkerInfo()
-	{
-		if (WorkerProcess.IsValid())
-		{
-			FPlatformProcess::TerminateProc(WorkerProcess);
-			FPlatformProcess::CloseProc(WorkerProcess);
-		}
-	}
-};
-
 /** Results for a single compiled shader map. */
 struct FComputeKernelShaderMapCompileResults
 {
@@ -60,7 +20,6 @@ struct FComputeKernelShaderMapCompileResults
 	TArray<FShaderCommonCompileJobPtr> FinishedJobs;
 };
 
-
 /** Results for a single compiled and finalized shader map. */
 struct FComputeKernelShaderMapFinalizeResults : public FComputeKernelShaderMapCompileResults
 {
@@ -79,20 +38,13 @@ struct FComputeKernelShaderMapFinalizeResults : public FComputeKernelShaderMapCo
 class FComputeKernelShaderCompilationManager
 {
 public:
-	FComputeKernelShaderCompilationManager();
-	~FComputeKernelShaderCompilationManager();
-
 	void Tick(float DeltaSeconds);
 	void AddJobs(TArray<FShaderCommonCompileJobPtr> InNewJobs);
-	void ProcessAsyncResults();
-
 	void FinishCompilation(const TCHAR* InKernelName, const TArray<int32>& ShaderMapIdsToFinishCompiling);
 
 private:
+	void ProcessAsyncResults();
 	void ProcessCompiledComputeKernelShaderMaps(TMap<int32, FComputeKernelShaderMapFinalizeResults>& CompiledShaderMaps, float TimeBudget);
-	void RunCompileJobs();
-
-	void InitWorkerInfo();
 
 	TArray<FShaderCommonCompileJobPtr> JobQueue;
 
@@ -101,8 +53,6 @@ private:
 
 	/** Map from shader map id to results being finalized.  Used to track shader finalizations over multiple frames. */
 	TMap<int32, FComputeKernelShaderMapFinalizeResults> PendingFinalizeComputeKernelShaderMaps;
-
-	TArray<struct FComputeKernelShaderCompileWorkerInfo*> WorkerInfos;
 };
 
 extern FComputeKernelShaderCompilationManager GComputeKernelShaderCompilationManager;
diff --git a/Engine/Plugins/Runtime/ComputeFramework/Source/ComputeFramework/Private/ComputeFramework/ComputeKernelShared.h b/Engine/Plugins/Runtime/ComputeFramework/Source/ComputeFramework/Private/ComputeFramework/ComputeKernelShared.h
--- a/Engine/Plugins/Runtime/ComputeFramework/Source/ComputeFramework/Private/ComputeFramework/ComputeKernelShared.h
+++ b/Engine/Plugins/Runtime/ComputeFramework/Source/ComputeFramework/Private/ComputeFramework/ComputeKernelShared.h
@@ -253,9 +253,6 @@ private:
 	static TArray<FComputeKernelShaderMap*> AllComputeKernelShaderMaps;
 
 #if WITH_EDITOR
-	/** Next value for CompilingId. */
-	static uint32 NextCompilingId;
-
 	/** Tracks resources and their shader maps that need to be compiled but whose compilation is being deferred. */
 	static TMap<TRefCountPtr<FComputeKernelShaderMap>, TArray<FComputeKernelResource*> > ComputeKernelShaderMapsBeingCompiled;
 #endif
diff --git a/Engine/Plugins/Runtime/ComputeFramework/Source/ComputeFramework/Private/ComputeKernelShaderCompilationManager.cpp b/Engine/Plugins/Runtime/ComputeFramework/Source/ComputeFramework/Private/ComputeKernelShaderCompilationManager.cpp
--- a/Engine/Plugins/Runtime/ComputeFramework/Source/ComputeFramework/Private/ComputeKernelShaderCompilationManager.cpp
+++ b/Engine/Plugins/Runtime/ComputeFramework/Source/ComputeFramework/Private/ComputeKernelShaderCompilationManager.cpp
@@ -2,19 +2,18 @@
 
 #include "ComputeFramework/ComputeKernelShaderCompilationManager.h"
 
+#include "ComputeFramework/ComputeKernelShared.h"
+#include "DataDrivenShaderPlatformInfo.h"
 #include "HAL/FileManager.h"
 #include "Misc/Paths.h"
-#include "ComputeFramework/ComputeKernelShared.h"
 #include "RHIShaderFormatDefinitions.inl"
 #include "ShaderCompiler.h"
-#include "DataDrivenShaderPlatformInfo.h"
 
 #if WITH_EDITOR
 #include "Interfaces/IShaderFormat.h"
 #include "Interfaces/ITargetPlatformManagerModule.h"
 #endif
 
-
 DEFINE_LOG_CATEGORY_STATIC(LogComputeKernelShaderCompiler, All, All);
 
 static int32 GShowComputeKernelShaderWarnings = 0;
@@ -30,214 +29,82 @@ FComputeKernelShaderCompilationManager GComputeKernelShaderCompilationManager;
 
 void FComputeKernelShaderCompilationManager::Tick(float DeltaSeconds)
 {
-	RunCompileJobs();
 	ProcessAsyncResults();
 }
 
-FComputeKernelShaderCompilationManager::FComputeKernelShaderCompilationManager()
-{
-	
-}
-
-FComputeKernelShaderCompilationManager::~FComputeKernelShaderCompilationManager()
-{
-	for (FComputeKernelShaderCompileWorkerInfo* Info : WorkerInfos)
-	{
-		delete Info;
-	}
-
-	WorkerInfos.Empty();
-}
-
-void FComputeKernelShaderCompilationManager::RunCompileJobs()
+void FComputeKernelShaderCompilationManager::AddJobs(TArray<FShaderCommonCompileJobPtr> InNewJobs)
 {
 	check(IsInGameThread());
+	JobQueue.Append(InNewJobs);
 
-	InitWorkerInfo();
-
-	int32 NumActiveThreads = 0;
-
-	for (int32 WorkerIndex = 0; WorkerIndex < WorkerInfos.Num(); WorkerIndex++)
+	for (FShaderCommonCompileJobPtr& Job : InNewJobs)
 	{
-		FComputeKernelShaderCompileWorkerInfo& CurrentWorkerInfo = *WorkerInfos[WorkerIndex];
-
-		// If this worker doesn't have any queued jobs, look for more in the input queue
-		if (CurrentWorkerInfo.QueuedJobs.Num() == 0)
-		{
-			check(!CurrentWorkerInfo.bComplete);
-
-			if (JobQueue.Num() > 0)
-			{
-				bool bAddedLowLatencyTask = false;
-				int32 JobIndex = 0;
+		FComputeKernelShaderMapCompileResults& ShaderMapInfo = ComputeKernelShaderMapJobs.FindOrAdd(Job->Id);
+		ShaderMapInfo.NumJobsQueued++;
 
-				// Try to grab up to MaxShaderJobBatchSize jobs
-				// Don't put more than one low latency task into a batch
-				for (; JobIndex < JobQueue.Num(); JobIndex++)
-				{
-					CurrentWorkerInfo.QueuedJobs.Add(JobQueue[JobIndex]);
-				}
+		FShaderCompileJob* CurrentJob = Job->GetSingleShaderJob();
 
-				// Update the worker state as having new tasks that need to be issued					
-				// don't reset worker app ID, because the shadercompilerworkers don't shutdown immediately after finishing a single job queue.
-				CurrentWorkerInfo.bIssuedTasksToWorker = true;
-				CurrentWorkerInfo.bLaunchedWorker = true;
-				CurrentWorkerInfo.StartTime = FPlatformTime::Seconds();
-				JobQueue.RemoveAt(0, JobIndex);
-			}
-		}
+		CurrentJob->Input.DumpDebugInfoRootPath = GShaderCompilingManager->GetAbsoluteShaderDebugInfoDirectory() / CurrentJob->Input.ShaderPlatformName.ToString();
+		FPaths::NormalizeDirectoryName(CurrentJob->Input.DumpDebugInfoRootPath);
 
-		if (CurrentWorkerInfo.bIssuedTasksToWorker && CurrentWorkerInfo.bLaunchedWorker)
+		CurrentJob->Input.DebugExtension.Empty();
+		CurrentJob->Input.DumpDebugInfoPath.Empty();
+		if (GShaderCompilingManager->GetDumpShaderDebugInfo() == FShaderCompilingManager::EDumpShaderDebugInfo::Always)
 		{
-			NumActiveThreads++;
+			CurrentJob->Input.DumpDebugInfoRootPath = GShaderCompilingManager->CreateShaderDebugInfoPath(CurrentJob->Input);
 		}
+	}
 
-		if (CurrentWorkerInfo.QueuedJobs.Num() > 0)
-		{
-			for (int32 JobIndex = 0; JobIndex < CurrentWorkerInfo.QueuedJobs.Num(); JobIndex++)
-			{
-				FShaderCompileJob& CurrentJob = static_cast<FShaderCompileJob&>(*CurrentWorkerInfo.QueuedJobs[JobIndex].GetReference());
-
-				check(!CurrentJob.bFinalized);
-				CurrentJob.bFinalized = true;
-
-				static ITargetPlatformManagerModule& TPM = GetTargetPlatformManagerRef();
-
-				const FName Format = LegacyShaderPlatformToShaderFormat(EShaderPlatform(CurrentJob.Input.Target.Platform));
-				const IShaderFormat* Compiler = TPM.FindShaderFormat(Format);
-
-				if (!Compiler)
-				{
-					UE_LOG(LogComputeKernelShaderCompiler, Fatal, TEXT("Can't compile shaders for format %s, couldn't load compiler dll"), *Format.ToString());
-				}
-				CA_ASSUME(Compiler != nullptr);
-
-				UE_LOG(LogComputeKernelShaderCompiler, Log, TEXT("Compile Job processing... %s"), *CurrentJob.Input.DebugGroupName);
-
-				CurrentJob.Input.DumpDebugInfoRootPath = GShaderCompilingManager->GetAbsoluteShaderDebugInfoDirectory() / CurrentJob.Input.ShaderPlatformName.ToString();
-				FPaths::NormalizeDirectoryName(CurrentJob.Input.DumpDebugInfoRootPath);
-				const FShaderCompilingManager::EDumpShaderDebugInfo DumpShaderDebugInfo = GShaderCompilingManager->GetDumpShaderDebugInfo();
-				CurrentJob.Input.DebugExtension.Empty();
-				if (DumpShaderDebugInfo == FShaderCompilingManager::EDumpShaderDebugInfo::Always)
-				{
-					CurrentJob.Input.DumpDebugInfoRootPath = GShaderCompilingManager->CreateShaderDebugInfoPath(CurrentJob.Input);
-				}
-
-				if (IsValidRef(CurrentJob.Input.SharedEnvironment))
-				{
-					// Merge the shared environment into the per-shader environment before calling into the compile function
-					// Normally this happens in the worker
-					CurrentJob.Input.Environment.Merge(*CurrentJob.Input.SharedEnvironment);
-				}
-
-				// Compile the shader directly through the platform dll (directly from the shader dir as the working directory)
-				Compiler->CompileShader(Format, CurrentJob.Input, CurrentJob.Output, FString(FPlatformProcess::ShaderDir()));
-
-				CurrentJob.bSucceeded = CurrentJob.Output.bSucceeded;
+	GShaderCompilingManager->SubmitJobs(InNewJobs, FString(), FString());
+}
 
-				// Recompile the shader to dump debug info if desired
-				if (GShaderCompilingManager->ShouldRecompileToDumpShaderDebugInfo(CurrentJob))
-				{
-					CurrentJob.Input.DumpDebugInfoPath = GShaderCompilingManager->CreateShaderDebugInfoPath(CurrentJob.Input);
-					CurrentJob.Output = FShaderCompilerOutput();
-					Compiler->CompileShader(Format, CurrentJob.Input, CurrentJob.Output, FString(FPlatformProcess::ShaderDir()));
-				}
+void FComputeKernelShaderCompilationManager::ProcessAsyncResults()
+{
+	check(IsInGameThread());
 
-				if (CurrentJob.Output.bSucceeded)
-				{
-					// Generate a hash of the output and cache it
-					// The shader processing this output will use it to search for existing FShaderResources
-					CurrentJob.Output.GenerateOutputHash();
-					UE_LOG(LogComputeKernelShaderCompiler, Log, TEXT("GPU shader compile succeeded. Id %d"), CurrentJob.Id);
-				}
-				else
-				{
-					UE_LOG(LogComputeKernelShaderCompiler, Log, TEXT("ERROR: GPU shader compile failed! Id %d"), CurrentJob.Id);
-				}
+	// Process the results from the shader compile worker
+	for (int32 JobIndex = JobQueue.Num() - 1; JobIndex >= 0; JobIndex--)
+	{
+		auto CurrentJob = JobQueue[JobIndex]->GetSingleShaderJob();
 
-				CurrentWorkerInfo.bComplete = true;
-			}
+		if (!CurrentJob->bReleased)
+		{
+			continue;
 		}
-	}
 
-	for (int32 WorkerIndex = 0; WorkerIndex < WorkerInfos.Num(); WorkerIndex++)
-	{
-		FComputeKernelShaderCompileWorkerInfo& CurrentWorkerInfo = *WorkerInfos[WorkerIndex];
-		if (CurrentWorkerInfo.bComplete)
+		CurrentJob->bSucceeded = CurrentJob->Output.bSucceeded;
+		if (CurrentJob->Output.bSucceeded)
 		{
-			for (int32 JobIndex = 0; JobIndex < CurrentWorkerInfo.QueuedJobs.Num(); JobIndex++)
-			{
-				FComputeKernelShaderMapCompileResults& ShaderMapResults = ComputeKernelShaderMapJobs.FindChecked(CurrentWorkerInfo.QueuedJobs[JobIndex]->Id);
-				ShaderMapResults.FinishedJobs.Add(CurrentWorkerInfo.QueuedJobs[JobIndex]);
-				ShaderMapResults.bAllJobsSucceeded = ShaderMapResults.bAllJobsSucceeded && CurrentWorkerInfo.QueuedJobs[JobIndex]->bSucceeded;
-			}
+			UE_LOG(LogComputeKernelShaderCompiler, Log, TEXT("GPU shader compile succeeded. Id %d"), CurrentJob->Id);
 		}
-
-		CurrentWorkerInfo.bComplete = false;
-		CurrentWorkerInfo.QueuedJobs.Empty();
-	}
-}
-
-void FComputeKernelShaderCompilationManager::InitWorkerInfo()
-{
-	if (WorkerInfos.Num() == 0) // Check to see if it has been initialized or not
-	{
-		// Ew. Should we just use FShaderCompilingManager's workers instead? Is that safe?
-		const int32 NumVirtualCores = FPlatformMisc::NumberOfCoresIncludingHyperthreads();
-		const uint32 NumComputeKernelShaderCompilingThreads = FMath::Min(NumVirtualCores - 1, 4);
-
-		for (uint32 WorkerIndex = 0; WorkerIndex < NumComputeKernelShaderCompilingThreads; WorkerIndex++)
+		else
 		{
-			WorkerInfos.Add(new FComputeKernelShaderCompileWorkerInfo());
+			UE_LOG(LogComputeKernelShaderCompiler, Log, TEXT("ERROR: GPU shader compile failed! Id %d"), CurrentJob->Id);
 		}
-	}	
-}
 
-void FComputeKernelShaderCompilationManager::AddJobs(TArray<FShaderCommonCompileJobPtr> InNewJobs)
-{
-	for (auto& Job : InNewJobs)
-	{
-		FComputeKernelShaderMapCompileResults& ShaderMapInfo = ComputeKernelShaderMapJobs.FindOrAdd(Job->Id);
-		//@todo : Apply shader map isn't used for now with this compile manager. Should be merged to have a generic shader compiler
-		ShaderMapInfo.NumJobsQueued++;
+		FComputeKernelShaderMapCompileResults& ShaderMapResults = ComputeKernelShaderMapJobs.FindChecked(CurrentJob->Id);
+		ShaderMapResults.FinishedJobs.Add(JobQueue[JobIndex]);
+		ShaderMapResults.bAllJobsSucceeded = ShaderMapResults.bAllJobsSucceeded && CurrentJob->bSucceeded;
+		JobQueue.RemoveAt(JobIndex);
 	}
 
-	JobQueue.Append(InNewJobs);
-}
-
-
-void FComputeKernelShaderCompilationManager::ProcessAsyncResults()
-{
-	int32 NumCompilingComputeKernelShaderMaps = 0;
-	TArray<int32> ShaderMapsToRemove;
-
-	// Get all ComputeKernel shader maps to finalize
-	//
 	for (TMap<int32, FComputeKernelShaderMapCompileResults>::TIterator It(ComputeKernelShaderMapJobs); It; ++It)
 	{
 		const FComputeKernelShaderMapCompileResults& Results = It.Value();
 
 		if (Results.FinishedJobs.Num() == Results.NumJobsQueued)
 		{
-			ShaderMapsToRemove.Add(It.Key());
-			PendingFinalizeComputeKernelShaderMaps.Add(It.Key(), FComputeKernelShaderMapFinalizeResults(Results));
+			PendingFinalizeComputeKernelShaderMaps.Add(It.Key(), FComputeKernelShaderMapCompileResults(Results));
+			ComputeKernelShaderMapJobs.Remove(It.Key());
 		}
 	}
 
-	for (int32 RemoveIndex = 0; RemoveIndex < ShaderMapsToRemove.Num(); RemoveIndex++)
-	{
-		ComputeKernelShaderMapJobs.Remove(ShaderMapsToRemove[RemoveIndex]);
-	}
-
-	NumCompilingComputeKernelShaderMaps = ComputeKernelShaderMapJobs.Num();
-
 	if (PendingFinalizeComputeKernelShaderMaps.Num() > 0)
 	{
-		ProcessCompiledComputeKernelShaderMaps(PendingFinalizeComputeKernelShaderMaps, 0.1f);
+		ProcessCompiledComputeKernelShaderMaps(PendingFinalizeComputeKernelShaderMaps, 10);
 	}
 }
 
-
 void FComputeKernelShaderCompilationManager::ProcessCompiledComputeKernelShaderMaps(
 	TMap<int32, FComputeKernelShaderMapFinalizeResults>& CompiledShaderMaps,
 	float TimeBudget)
@@ -446,12 +313,11 @@ void FComputeKernelShaderCompilationManager::ProcessCompiledComputeKernelShaderM
 	}
 }
 
-
 void FComputeKernelShaderCompilationManager::FinishCompilation(const TCHAR* InKernelName, const TArray<int32>& ShaderMapIdsToFinishCompiling)
 {
 	check(!FPlatformProperties::RequiresCookedData());
 
-	RunCompileJobs();	// since we don't async compile through another process, this will run all oustanding jobs
+	GShaderCompilingManager->FinishCompilation(NULL, ShaderMapIdsToFinishCompiling);
 	ProcessAsyncResults();	// grab compiled shader maps and assign them to their resources
 
 	check(ComputeKernelShaderMapJobs.Num() == 0);
diff --git a/Engine/Plugins/Runtime/ComputeFramework/Source/ComputeFramework/Private/ComputeKernelShaderMap.cpp b/Engine/Plugins/Runtime/ComputeFramework/Source/ComputeFramework/Private/ComputeKernelShaderMap.cpp
--- a/Engine/Plugins/Runtime/ComputeFramework/Source/ComputeFramework/Private/ComputeKernelShaderMap.cpp
+++ b/Engine/Plugins/Runtime/ComputeFramework/Source/ComputeFramework/Private/ComputeKernelShaderMap.cpp
@@ -42,11 +42,6 @@ TMap<FComputeKernelShaderMapId, FComputeKernelShaderMap*> FComputeKernelShaderMa
 TArray<FComputeKernelShaderMap*> FComputeKernelShaderMap::AllComputeKernelShaderMaps;
 
 #if WITH_EDITOR
-
-// The Id of 0 is reserved for global shaders
-uint32 FComputeKernelShaderMap::NextCompilingId = 2;
-
-
 /** 
  * Tracks FComputeKernelResource and their shader maps that are being compiled.
  * Uses a TRefCountPtr as this will be the only reference to a shader map while it is being compiled.
@@ -55,7 +50,6 @@ TMap<TRefCountPtr<FComputeKernelShaderMap>, TArray<FComputeKernelResource*> > FC
 #endif // WITH_EDITOR
 
 
-
 static inline bool ShouldCacheComputeKernelShader(const FComputeKernelShaderType* InShaderType, EShaderPlatform InPlatform, const FComputeKernelResource* InKernel)
 {
 	return InShaderType->ShouldCache(InPlatform, InKernel) && InKernel->ShouldCache(InPlatform, InShaderType);
@@ -342,13 +336,10 @@ void FComputeKernelShaderMap::Compile(
 		else
 		{
 			// Assign a unique identifier so that shaders from this shader map can be associated with it after a deferred compile
-			CompilingId = NextCompilingId;
+			CompilingId = FShaderCommonCompileJob::GetNextJobId();
 			UE_LOG(LogComputeFramework, Log, TEXT("CompilingId = %p %d"), InKernel, CompilingId);
 			InKernel->AddCompileId(CompilingId);
 
-			check(NextCompilingId < UINT_MAX);
-			NextCompilingId++;
-  
 			// Setup the compilation environment.
 			InKernel->SetupShaderCompilationEnvironment(InPlatform, *InCompilationEnvironment);
   
