ComputeFramework: Don't compile compute graphs on PostLoad.
Defer compilation to first use.

diff --git a/Engine/Plugins/Runtime/ComputeFramework/Source/ComputeFramework/Private/ComputeFramework.cpp b/Engine/Plugins/Runtime/ComputeFramework/Source/ComputeFramework/Private/ComputeFramework.cpp
--- a/Engine/Plugins/Runtime/ComputeFramework/Source/ComputeFramework/Private/ComputeFramework.cpp
+++ b/Engine/Plugins/Runtime/ComputeFramework/Source/ComputeFramework/Private/ComputeFramework.cpp
@@ -23,6 +23,14 @@ static FAutoConsoleVariableRef CVarComputeFrameworkEnable(
 	ECVF_Scalability | ECVF_RenderThreadSafe
 	);
 
+static int32 GComputeFrameworkDeferredCompilation = 1;
+static FAutoConsoleVariableRef CVarComputeFrameworkDeferredCompilation(
+	TEXT("r.ComputeFramework.DeferredCompilation"),
+	GComputeFrameworkDeferredCompilation,
+	TEXT("Compile compute graphs on first usage instead of on PostLoad().\n"),
+	ECVF_Default
+);
+
 FAutoConsoleCommand CmdRebuildComputeGraphs(
 	TEXT("r.ComputeFramework.RebuildComputeGraphs"),
 	TEXT("Force all loaded UComputeGraph objects to rebuild."),
@@ -41,6 +49,15 @@ namespace ComputeFramework
 		return GComputeFrameworkEnable > 0 && FComputeFrameworkModule::GetComputeSystem() != nullptr;
 	}
 
+	bool IsDeferredCompilation()
+	{
+#if WITH_EDITOR
+		return GComputeFrameworkDeferredCompilation != 0;
+#else
+		return false;
+#endif
+	}
+
 	void RebuildComputeGraphs()
 	{
 #if WITH_EDITOR
diff --git a/Engine/Plugins/Runtime/ComputeFramework/Source/ComputeFramework/Private/ComputeGraph.cpp b/Engine/Plugins/Runtime/ComputeFramework/Source/ComputeFramework/Private/ComputeGraph.cpp
--- a/Engine/Plugins/Runtime/ComputeFramework/Source/ComputeFramework/Private/ComputeGraph.cpp
+++ b/Engine/Plugins/Runtime/ComputeFramework/Source/ComputeFramework/Private/ComputeGraph.cpp
@@ -80,7 +80,10 @@ void UComputeGraph::PostLoad()
 	}
 #endif
 
-	UpdateResources();
+	if (!ComputeFramework::IsDeferredCompilation())
+	{
+		UpdateResources();
+	}
 }
 
 bool UComputeGraph::ValidateGraph(FString* OutErrors)
@@ -600,7 +603,7 @@ void UComputeGraph::CacheResourceShadersForRendering(uint32 CompilationFlags)
 
 			KernelResource->OnCompilationComplete().BindUObject(this, &UComputeGraph::ShaderCompileCompletionCallback);
 
-			CacheShadersForResource(ShaderPlatform, nullptr, CompilationFlags | uint32(EComputeKernelCompilationFlags::Force), KernelResource);
+			CacheShadersForResource(ShaderPlatform, nullptr, CompilationFlags, KernelResource);
 		}
 	}
 }
diff --git a/Engine/Plugins/Runtime/ComputeFramework/Source/ComputeFramework/Private/ComputeGraphInstance.cpp b/Engine/Plugins/Runtime/ComputeFramework/Source/ComputeFramework/Private/ComputeGraphInstance.cpp
--- a/Engine/Plugins/Runtime/ComputeFramework/Source/ComputeFramework/Private/ComputeGraphInstance.cpp
+++ b/Engine/Plugins/Runtime/ComputeFramework/Source/ComputeFramework/Private/ComputeGraphInstance.cpp
@@ -52,8 +52,12 @@ bool FComputeGraphInstance::EnqueueWork(UComputeGraph* InComputeGraph, FSceneInt
 	}
 
 	FComputeGraphRenderProxy const* GraphRenderProxy = InComputeGraph->GetRenderProxy();
-	if (!ensure(GraphRenderProxy))
+	if (GraphRenderProxy == nullptr)
 	{
+		// This can happen if we have deferred compilation.
+		// Trigger compilation now.
+		ensure(ComputeFramework::IsDeferredCompilation());
+		InComputeGraph->UpdateResources();
 		return false;
 	}
 
diff --git a/Engine/Plugins/Runtime/ComputeFramework/Source/ComputeFramework/Public/ComputeFramework/ComputeFramework.h b/Engine/Plugins/Runtime/ComputeFramework/Source/ComputeFramework/Public/ComputeFramework/ComputeFramework.h
--- a/Engine/Plugins/Runtime/ComputeFramework/Source/ComputeFramework/Public/ComputeFramework/ComputeFramework.h
+++ b/Engine/Plugins/Runtime/ComputeFramework/Source/ComputeFramework/Public/ComputeFramework/ComputeFramework.h
@@ -20,6 +20,9 @@ namespace ComputeFramework
 	/** Returns true if ComputeFramework is currently enabled. */
 	COMPUTEFRAMEWORK_API bool IsEnabled();
 
+	/** Returns true if compute graphs are compiled on first use instead of on PostLoad(). */
+	bool IsDeferredCompilation();
+
 	/** Rebuild all compute graphs. */
 	COMPUTEFRAMEWORK_API void RebuildComputeGraphs();
 
