Remove one unnecessary layer of file indirection in the compute framework shader compilation.

diff --git a/Engine/Plugins/Runtime/ComputeFramework/Source/ComputeFramework/Private/ComputeKernelShaderMap.cpp b/Engine/Plugins/Runtime/ComputeFramework/Source/ComputeFramework/Private/ComputeKernelShaderMap.cpp
--- a/Engine/Plugins/Runtime/ComputeFramework/Source/ComputeFramework/Private/ComputeKernelShaderMap.cpp
+++ b/Engine/Plugins/Runtime/ComputeFramework/Source/ComputeFramework/Private/ComputeKernelShaderMap.cpp
@@ -142,10 +142,8 @@ void FComputeKernelShaderType::BeginCompileShader(
 		EShaderCompileJobPriority::Normal
 		);
 
-	FString GeneratedSourceFilePath = FString::Printf(TEXT("/Plugin/ComputeFramework/Generated/Kernel_%s.ush"), *InKernel->GetEntryPoint());
-	FString GeneratedStubSource = FString::Printf(TEXT("#include \"%s\"\n"), *GeneratedSourceFilePath);
-	TCHAR const* GeneratedStubFilePath = TEXT("/Plugin/ComputeFramework/Generated/ComputeKernel.ush");
 	TCHAR const* SourceFilePath = TEXT("/Plugin/ComputeFramework/Private/ComputeKernel.usf");
+	TCHAR const* GeneratedFilePath = TEXT("/Plugin/ComputeFramework/Generated/ComputeKernel.ush");
 
 	NewJob->ShaderParameters = MakeShared<const FParameters, ESPMode::ThreadSafe>(*InKernel->GetShaderParamMetadata());
 	NewJob->Input.SharedEnvironment = InCompilationEnvironment;
@@ -153,8 +151,7 @@ void FComputeKernelShaderType::BeginCompileShader(
 	NewJob->Input.ShaderFormat = LegacyShaderPlatformToShaderFormat(InPlatform);
 	NewJob->Input.VirtualSourceFilePath = SourceFilePath;
 	NewJob->Input.EntryPointName = InKernel->GetEntryPoint();
-	NewJob->Input.Environment.IncludeVirtualPathToContentsMap.Add(GeneratedSourceFilePath, InKernel->GetHLSLSource());
-	NewJob->Input.Environment.IncludeVirtualPathToContentsMap.Add(GeneratedStubFilePath, GeneratedStubSource);
+	NewJob->Input.Environment.IncludeVirtualPathToContentsMap.Add(GeneratedFilePath, InKernel->GetHLSLSource());
 	UE_LOG(LogComputeFramework, Verbose, TEXT("%s"), *InKernel->GetHLSLSource());
 	
 	AddUniformBufferIncludesToEnvironment(NewJob->Input.Environment, InPlatform);
