ComputeFramework: Fix compile error links for UObject shader paths.
Tidied logging during shader compilation.
Optimus: Add support for opening file in visual studio from shader compile errors in Editor or VS log.

diff --git a/Engine/Plugins/Animation/DeformerGraph/Source/OptimusCore/Private/Nodes/OptimusNode_ComputeKernelBase.cpp b/Engine/Plugins/Animation/DeformerGraph/Source/OptimusCore/Private/Nodes/OptimusNode_ComputeKernelBase.cpp
--- a/Engine/Plugins/Animation/DeformerGraph/Source/OptimusCore/Private/Nodes/OptimusNode_ComputeKernelBase.cpp
+++ b/Engine/Plugins/Animation/DeformerGraph/Source/OptimusCore/Private/Nodes/OptimusNode_ComputeKernelBase.cpp
@@ -330,6 +330,9 @@ FString UOptimusNode_ComputeKernelBase::GetCookedKernelSource(
 	Source.ReplaceInline(TEXT("\r"), TEXT(""));
 #endif
 
+	FString ShaderPathName = InObjectPathName;
+	Optimus::ConvertObjectPathToShaderFilePath(ShaderPathName);
+
 	const bool bHasKernelKeyword = Source.Contains(TEXT("KERNEL"));
 	
 	const FString KernelFunc = FString::Printf(
@@ -345,7 +348,7 @@ FString UOptimusNode_ComputeKernelBase::GetCookedKernelSource(
 				"#line 1 \"%s\"\n"
 				"%s\n\n"
 				"%s { __kernel_func(DTid.x); }\n"
-				), *InObjectPathName, *Source, *KernelFunc);
+				), *ShaderPathName, *Source, *KernelFunc);
 	}
 	else
 	{
@@ -357,7 +360,7 @@ FString UOptimusNode_ComputeKernelBase::GetCookedKernelSource(
 			"#line 1 \"%s\"\n"
 			"%s\n"
 			"}\n"
-			), *KernelFunc, *InObjectPathName, *Source);
+			), *KernelFunc, *ShaderPathName, *Source);
 	}
 }
 
diff --git a/Engine/Plugins/Animation/DeformerGraph/Source/OptimusCore/Private/OptimusComputeGraph.cpp b/Engine/Plugins/Animation/DeformerGraph/Source/OptimusCore/Private/OptimusComputeGraph.cpp
--- a/Engine/Plugins/Animation/DeformerGraph/Source/OptimusCore/Private/OptimusComputeGraph.cpp
+++ b/Engine/Plugins/Animation/DeformerGraph/Source/OptimusCore/Private/OptimusComputeGraph.cpp
@@ -3,12 +3,15 @@
 #include "OptimusComputeGraph.h"
 
 #include "Components/MeshComponent.h"
+#include "ComputeFramework/ComputeKernelCompileResult.h"
+#include "HAL/FileManager.h"
 #include "Internationalization/Regex.h"
 #include "IOptimusComputeKernelProvider.h"
 #include "IOptimusShaderTextProvider.h"
 #include "IOptimusValueProvider.h"
 #include "OptimusDeformer.h"
 #include "OptimusCoreModule.h"
+#include "OptimusHelpers.h"
 #include "OptimusNode.h"
 #include "OptimusObjectVersion.h"
 #include "Misc/UObjectToken.h"
@@ -35,50 +38,64 @@ void UOptimusComputeGraph::PostLoad()
 	}
 }
 
-static EOptimusDiagnosticLevel ProcessCompilationMessage(UOptimusDeformer* InOwner,	UOptimusNode const* InKernelNode, FString const& InMessage)
+static EOptimusDiagnosticLevel ProcessCompilationMessage(UOptimusDeformer* InOwner,	UOptimusNode const* InKernelNode, FString const& InKernelFriendlyName, FComputeKernelCompileMessage const& InMessage)
 {
 	FOptimusCompilerDiagnostic Diagnostic;
 
-	// "/Engine/Generated/ComputeFramework/Kernel_LinearBlendSkinning.usf(19,39-63):  error X3013: 'DI000_ReadNumVertices': no matching 1 parameter function"	
-	// "OptimusNode_ComputeKernel_2(1,42):  error X3004: undeclared identifier 'a'"
-
-	// TODO: Parsing diagnostics rightfully belongs at the shader compiler level, especially if the shader compiler is rewriting.
-	static const FRegexPattern MessagePattern(TEXT(R"(^\s*(.*?)\((\d+),(\d+)(-(\d+))?\):\s*(error|warning)\s+[A-Z0-9]+:\s*(.*)$)"));
-	FRegexMatcher Matcher(MessagePattern, InMessage);
-
-	if (!Matcher.FindNext())
+	if (InMessage.Type == FComputeKernelCompileMessage::EMessageType::Error)
+	{
+		Diagnostic.Level = EOptimusDiagnosticLevel::Error;
+	}
+	else if (InMessage.Type == FComputeKernelCompileMessage::EMessageType::Warning)
+	{
+		Diagnostic.Level = EOptimusDiagnosticLevel::Warning;
+	}
+	else if (InMessage.Type == FComputeKernelCompileMessage::EMessageType::Info)
 	{
 		Diagnostic.Level = EOptimusDiagnosticLevel::Info;
-		Diagnostic.Message = FText::FromString(InMessage);
-		//Diagnostic.Object = InKernelNode;
 	}
-	else
+
+	Diagnostic.Line = InMessage.Line;
+	Diagnostic.ColumnStart = InMessage.ColumnStart;
+	Diagnostic.ColumnEnd = InMessage.ColumnEnd;
+
+	// If error path is a UObject, then stpre a reference to the related object.
+	FString Path = InMessage.VirtualFilePath;
+	if (Optimus::ConvertShaderFilePathToObjectPath(Path))
 	{
-		const FString SeverityStr = Matcher.GetCaptureGroup(6);
-		if (SeverityStr == TEXT("warning"))
+		Diagnostic.Object = StaticFindObject(nullptr, nullptr, *Path, true);
+		if (UObject const* ObjectPtr = Diagnostic.Object.Get())
 		{
-			Diagnostic.Level = EOptimusDiagnosticLevel::Warning;
+			// Special case path display if object is the current kernel.
+			Path = (ObjectPtr == InKernelNode) ? InKernelFriendlyName : ObjectPtr->GetName();
 		}
-		else if (SeverityStr == TEXT("error"))
+	}
+	
+	// If error path is a real file, then store its absolute path on disk.
+	if (!InMessage.RealFilePath.IsEmpty())
+	{
+		Diagnostic.AbsoluteFilePath = IFileManager::Get().ConvertToAbsolutePathForExternalAppForRead(*InMessage.RealFilePath);
+	}
+
+	FString Message;
+	if (!Path.IsEmpty())
+	{
+		Message = Path;
+		if (InMessage.Line != -1)
 		{
-			Diagnostic.Level = EOptimusDiagnosticLevel::Error;
+			if (InMessage.ColumnStart == InMessage.ColumnEnd)
+			{
+				Message += FString::Printf(TEXT(" (%d,%d)"), InMessage.Line, InMessage.ColumnStart);
+			}
+			else
+			{
+				Message += FString::Printf(TEXT(" (%d,%d-%d)"), InMessage.Line, InMessage.ColumnStart, InMessage.ColumnEnd);
+			}
 		}
-
-		const FString Path = Matcher.GetCaptureGroup(1);
-		Diagnostic.Object = StaticFindObject(nullptr, nullptr, *Path, true);
-
-		const FString MessageStr = Matcher.GetCaptureGroup(7);
-		const bool bShowPathInMessage = Diagnostic.Object == nullptr;
-		Diagnostic.Message = FText::FromString(bShowPathInMessage ? FString::Printf(TEXT("%s: %s"), *Path, *MessageStr) : MessageStr);
-
-		const int32 LineNumber = FCString::Atoi(*Matcher.GetCaptureGroup(2));
-		const int32 ColumnStart = FCString::Atoi(*Matcher.GetCaptureGroup(3));
-		const FString ColumnEndStr = Matcher.GetCaptureGroup(5);
-		const int32 ColumnEnd = ColumnEndStr.IsEmpty() ? ColumnStart : FCString::Atoi(*ColumnEndStr);
-		Diagnostic.Line = LineNumber;
-		Diagnostic.ColumnStart = ColumnStart;
-		Diagnostic.ColumnEnd = ColumnEnd;
+		Message += TEXT(": ");
 	}
+	Message += InMessage.Text;
+	Diagnostic.Message = FText::FromString(Message);
 
 	if (InOwner)
 	{
@@ -88,7 +105,7 @@ static EOptimusDiagnosticLevel ProcessCompilationMessage(UOptimusDeformer* InOwn
 	return Diagnostic.Level;
 }
 
-void UOptimusComputeGraph::OnKernelCompilationComplete(int32 InKernelIndex, const TArray<FString>& InCompileOutputMessages)
+void UOptimusComputeGraph::OnKernelCompilationComplete(int32 InKernelIndex, FComputeKernelCompileResults const& InCompileResults)
 {
 	// Find the Optimus objects from the raw kernel index.
 	if (KernelToNode.IsValidIndex(InKernelIndex))
@@ -98,11 +115,12 @@ void UOptimusComputeGraph::OnKernelCompilationComplete(int32 InKernelIndex, cons
 		// Make sure the node hasn't been GC'd.
 		if (UOptimusNode* Node = const_cast<UOptimusNode*>(KernelToNode[InKernelIndex].Get()))
 		{
+			FString FriendlyName = Owner->GetName() / GetFName().GetPlainNameString() / Node->GetDisplayName().ToString();
 			EOptimusDiagnosticLevel DiagnosticLevel = EOptimusDiagnosticLevel::None;
 
-			for (FString const& CompileOutputMessage : InCompileOutputMessages)
+			for (FComputeKernelCompileMessage const& CompileOutputMessage : InCompileResults.Messages)
 			{
-				EOptimusDiagnosticLevel MessageDiagnosticLevel = ProcessCompilationMessage(Owner, Node, CompileOutputMessage);
+				EOptimusDiagnosticLevel MessageDiagnosticLevel = ProcessCompilationMessage(Owner, Node, FriendlyName, CompileOutputMessage);
 				if (MessageDiagnosticLevel > DiagnosticLevel)
 				{
 					DiagnosticLevel = MessageDiagnosticLevel;
diff --git a/Engine/Plugins/Animation/DeformerGraph/Source/OptimusCore/Private/OptimusComputeGraph.h b/Engine/Plugins/Animation/DeformerGraph/Source/OptimusCore/Private/OptimusComputeGraph.h
--- a/Engine/Plugins/Animation/DeformerGraph/Source/OptimusCore/Private/OptimusComputeGraph.h
+++ b/Engine/Plugins/Animation/DeformerGraph/Source/OptimusCore/Private/OptimusComputeGraph.h
@@ -22,7 +22,7 @@ public:
 	void PostLoad() override;
 
 	// UComputeGraph overrides
-	void OnKernelCompilationComplete(int32 InKernelIndex, const TArray<FString>& InCompileOutputMessages) override;
+	void OnKernelCompilationComplete(int32 InKernelIndex, FComputeKernelCompileResults const& InCompileResults) override;
 
 protected:
 	// Lookup into Graphs array from the UComputeGraph kernel index. 
diff --git a/Engine/Plugins/Animation/DeformerGraph/Source/OptimusCore/Private/OptimusHelpers.cpp b/Engine/Plugins/Animation/DeformerGraph/Source/OptimusCore/Private/OptimusHelpers.cpp
--- a/Engine/Plugins/Animation/DeformerGraph/Source/OptimusCore/Private/OptimusHelpers.cpp
+++ b/Engine/Plugins/Animation/DeformerGraph/Source/OptimusCore/Private/OptimusHelpers.cpp
@@ -110,3 +110,22 @@ FName Optimus::GetTypeName(UScriptStruct* InStructType, bool bInShouldGetUniqueN
 	return FName(*InStructType->GetStructCPPName());
 }
 
+void Optimus::ConvertObjectPathToShaderFilePath(FString& InOutPath)
+{
+	// Shader compiler recognizes "/UObject/..." path as special. 
+	// It doesn't validate file suffix etc.
+	InOutPath = TEXT("/UObject") + InOutPath;
+	// Shader compilation result parsing will break if it finds ':' where it doesn't expect.
+	InOutPath.ReplaceCharInline(TEXT(':'), TEXT('@'));
+}
+
+bool Optimus::ConvertShaderFilePathToObjectPath(FString& InOutPath)
+{
+	if (!InOutPath.RemoveFromStart(TEXT("/UObject")))
+	{
+		return false;
+	}
+
+	InOutPath.ReplaceCharInline(TEXT('@'), TEXT(':'));
+	return true;
+}
diff --git a/Engine/Plugins/Animation/DeformerGraph/Source/OptimusCore/Private/OptimusSource.cpp b/Engine/Plugins/Animation/DeformerGraph/Source/OptimusCore/Private/OptimusSource.cpp
--- a/Engine/Plugins/Animation/DeformerGraph/Source/OptimusCore/Private/OptimusSource.cpp
+++ b/Engine/Plugins/Animation/DeformerGraph/Source/OptimusCore/Private/OptimusSource.cpp
@@ -1,6 +1,7 @@
 // Copyright Epic Games, Inc. All Rights Reserved.
 
 #include "OptimusSource.h"
+#include "OptimusHelpers.h"
 
 #include UE_INLINE_GENERATED_CPP_BY_NAME(OptimusSource)
 
@@ -11,6 +12,14 @@ void UOptimusSource::SetSource(const FString& InText)
 	Modify();
 }
 
+
+FString UOptimusSource::GetSource() const
+{
+	FString ShaderPathName = GetPathName();
+	Optimus::ConvertObjectPathToShaderFilePath(ShaderPathName);
+	return FString::Printf(TEXT("#line 1 \"%s\"\n%s"), *ShaderPathName, *SourceText);
+}
+
 #if WITH_EDITOR	
 
 FString UOptimusSource::GetNameForShaderTextEditor() const
diff --git a/Engine/Plugins/Animation/DeformerGraph/Source/OptimusCore/Public/OptimusDiagnostic.h b/Engine/Plugins/Animation/DeformerGraph/Source/OptimusCore/Public/OptimusDiagnostic.h
--- a/Engine/Plugins/Animation/DeformerGraph/Source/OptimusCore/Public/OptimusDiagnostic.h
+++ b/Engine/Plugins/Animation/DeformerGraph/Source/OptimusCore/Public/OptimusDiagnostic.h
@@ -35,6 +35,9 @@ struct FOptimusCompilerDiagnostic
 	// Ending column (inclusive)
 	int32 ColumnEnd = INDEX_NONE;
 
+	// Absolute path to file associated with this diagnostic
+	FString AbsoluteFilePath;
+
 	// UObject associated with this diagnostic
 	TWeakObjectPtr<const UObject> Object;
 };
diff --git a/Engine/Plugins/Animation/DeformerGraph/Source/OptimusCore/Public/OptimusHelpers.h b/Engine/Plugins/Animation/DeformerGraph/Source/OptimusCore/Public/OptimusHelpers.h
--- a/Engine/Plugins/Animation/DeformerGraph/Source/OptimusCore/Public/OptimusHelpers.h
+++ b/Engine/Plugins/Animation/DeformerGraph/Source/OptimusCore/Public/OptimusHelpers.h
@@ -92,4 +92,9 @@ namespace Optimus
 
 	/** Helper function to remove guids from member property names for user defined structs	*/
 	OPTIMUSCORE_API FName GetMemberPropertyShaderName(UScriptStruct* InStruct, const FProperty* InMemberProperty);
+
+	/** Helpers to convert UObject path to a virtual shader paths. */
+	void ConvertObjectPathToShaderFilePath(FString& InOutPath);
+	/** Helpers to convert a virtual shader path back to a UObject path. Returns false if the virtual shader path wasn't recognized as a UObject path. */
+	bool ConvertShaderFilePathToObjectPath(FString& InOutPath);
 }
diff --git a/Engine/Plugins/Animation/DeformerGraph/Source/OptimusCore/Public/OptimusSource.h b/Engine/Plugins/Animation/DeformerGraph/Source/OptimusCore/Public/OptimusSource.h
--- a/Engine/Plugins/Animation/DeformerGraph/Source/OptimusCore/Public/OptimusSource.h
+++ b/Engine/Plugins/Animation/DeformerGraph/Source/OptimusCore/Public/OptimusSource.h
@@ -17,7 +17,7 @@ public:
 	void SetSource(const FString& InText);
 
 	// Begin UComputeSource interface.
-	FString GetSource() const override { return SourceText; }
+	FString GetSource() const override;
 	// End UComputeSource interface.
 
 	// Begin IOptimusShaderTextProvider interface.
diff --git a/Engine/Plugins/Animation/DeformerGraph/Source/OptimusEditor/OptimusEditor.Build.cs b/Engine/Plugins/Animation/DeformerGraph/Source/OptimusEditor/OptimusEditor.Build.cs
--- a/Engine/Plugins/Animation/DeformerGraph/Source/OptimusEditor/OptimusEditor.Build.cs
+++ b/Engine/Plugins/Animation/DeformerGraph/Source/OptimusEditor/OptimusEditor.Build.cs
@@ -37,7 +37,8 @@ namespace UnrealBuildTool.Rules
 					"Persona",
 					"MessageLog",
 					"PropertyEditor",
-					"KismetCompiler"
+					"KismetCompiler",
+					"SourceCodeAccess"
 				}
 			);
 
diff --git a/Engine/Plugins/Animation/DeformerGraph/Source/OptimusEditor/Private/OptimusEditor.cpp b/Engine/Plugins/Animation/DeformerGraph/Source/OptimusEditor/Private/OptimusEditor.cpp
--- a/Engine/Plugins/Animation/DeformerGraph/Source/OptimusEditor/Private/OptimusEditor.cpp
+++ b/Engine/Plugins/Animation/DeformerGraph/Source/OptimusEditor/Private/OptimusEditor.cpp
@@ -45,6 +45,8 @@
 #include "SkeletalRenderPublic.h"
 #include "ToolMenus.h"
 #include "OptimusShaderText.h"
+#include "ISourceCodeAccessModule.h"
+#include "ISourceCodeAccessor.h"
 
 
 #define LOCTEXT_NAMESPACE "OptimusEditor"
@@ -327,40 +329,54 @@ void FOptimusEditor::OnCompileMessage(FOptimusCompilerDiagnostic const& Diagnost
 
 	if (Diagnostic.Line != INDEX_NONE)
 	{
-		// Create message with line number.
-		TSharedRef<FTokenizedMessage> Message = FTokenizedMessage::Create(
-			Severity,
-			FText::Format(LOCTEXT("CompileMessageWithLine", "{0} (line {1})"), Diagnostic.Message, FText::AsNumber(Diagnostic.Line)));
+		TSharedRef<FTokenizedMessage> Message = FTokenizedMessage::Create(Severity,	Diagnostic.Message);
 
 		if (UObject const* TokenObject = Diagnostic.Object.Get())
 		{
-			// Add activation which opens a text editor tab for the text provider.
-			auto Activation = [Event = OnSelectedNodesChanged(), TabManagerPtr = GetTabManager(), TokenObject](const TSharedRef<class IMessageToken>&)
+			// Add activation which opens in text editor tab.
+			// Could improve by putting cursor at correct location.
+			Message->AddToken(FActionToken::Create(
+				LOCTEXT("OpenInShaderTextEditor", "Open in shader text editor."),
+				FText::FromString(TokenObject->GetName()),
+				FOnActionTokenExecuted::CreateLambda([Event = OnSelectedNodesChanged(), TabManagerPtr = GetTabManager(), TokenObject]
+					{
+						TabManagerPtr->TryInvokeTab(FOptimusEditorShaderTextEditorTabSummoner::TabId);
+						UObject* MutableTokenObject = const_cast<UObject*>(TokenObject);
+						Event.Broadcast({ MutableTokenObject });
+					})
+			));
+		}
+		else if (!Diagnostic.AbsoluteFilePath.IsEmpty())
+		{
+			// Add activation which opens an external text editor.
+			if (ISourceCodeAccessModule* SourceCodeAccessModule = FModuleManager::LoadModulePtr<ISourceCodeAccessModule>("SourceCodeAccess"))
 			{
-				TabManagerPtr->TryInvokeTab(FOptimusEditorShaderTextEditorTabSummoner::TabId);
-
-				UObject* MutableTokenObject = const_cast<UObject*>(TokenObject);
-				Event.Broadcast({ MutableTokenObject });
-			};
-			Message->AddToken(FUObjectToken::Create(TokenObject)->OnMessageTokenActivated(FOnMessageTokenActivated::CreateLambda(Activation)));
+				Message->AddToken(FActionToken::Create(
+					LOCTEXT("OpenInExternalEditor", "Open in external text editor."),
+					FText::FromString(Diagnostic.AbsoluteFilePath),
+					FOnActionTokenExecuted::CreateLambda([SourceCodeAccessModule, Path = Diagnostic.AbsoluteFilePath, Line = Diagnostic.Line, Column = Diagnostic.ColumnStart]
+						{
+							SourceCodeAccessModule->GetAccessor().OpenFileAtLine(Path, Line, Column);
+						})
+				));
+			}
 		}
 
-		CompilerResultsListing->AddMessage(Message);
+		CompilerResultsListing->AddMessage(Message, false);
 	}
 	else
 	{
-		// Create message.
 		TSharedRef<FTokenizedMessage> Message = FTokenizedMessage::Create(Severity, Diagnostic.Message);
 
 		if (UObject const* TokenObject = Diagnostic.Object.Get())
 		{
 			// Add a dummy activation since default behavior opens content browser.
-			// Could improve this to select node?
+			// Could improve this to select node in deformer graph view?
 			static auto DummyActivation = [](const TSharedRef<class IMessageToken>&) {};
 			Message->AddToken(FUObjectToken::Create(TokenObject)->OnMessageTokenActivated(FOnMessageTokenActivated::CreateLambda(DummyActivation)));
 		}
 	
-		CompilerResultsListing->AddMessage(Message);
+		CompilerResultsListing->AddMessage(Message, false);
 	}
 
 	// Add to diagnostic list which is referenced by syntax highlighters.
@@ -836,7 +852,7 @@ void FOptimusEditor::HandlePreviewMeshChanged(
 void FOptimusEditor::CreateMessageLog()
 {
 	FMessageLogModule& MessageLogModule = FModuleManager::LoadModuleChecked<FMessageLogModule>("MessageLog");
-	const FName LogName("LogComputeKernelShaderCompiler");
+	const FName LogName("LogOptimusCompiler");
 	if (MessageLogModule.IsRegisteredLogListing(LogName))
 	{
 		CompilerResultsListing = MessageLogModule.GetLogListing(LogName);
diff --git a/Engine/Plugins/Runtime/ComputeFramework/Source/ComputeFramework/Private/ComputeFramework/ComputeKernelShared.h b/Engine/Plugins/Runtime/ComputeFramework/Source/ComputeFramework/Private/ComputeFramework/ComputeKernelShared.h
--- a/Engine/Plugins/Runtime/ComputeFramework/Source/ComputeFramework/Private/ComputeFramework/ComputeKernelShared.h
+++ b/Engine/Plugins/Runtime/ComputeFramework/Source/ComputeFramework/Private/ComputeFramework/ComputeKernelShared.h
@@ -3,6 +3,7 @@
 #pragma once
 
 #include "CoreMinimal.h"
+#include "ComputeFramework/ComputeKernelCompileResult.h"
 #include "Engine/EngineTypes.h"
 #include "RenderResource.h"
 #include "RenderingThread.h"
@@ -362,8 +363,8 @@ public:
 	bool IsCompilationFinished() const;
 
 	// Accessors.
-	const TArray<FString>& GetCompileOutputMessages() const { return CompileOutputMessages; }
-	void SetCompileOutputMessages(TArray<FString> const& InCompileOutputMessages) { CompileOutputMessages = InCompileOutputMessages; }
+	FComputeKernelCompileResults const& GetCompilationResults() const { return CompilationResults; }
+	void SetCompilationResults(FComputeKernelCompileResults const& InCompilationResults) { CompilationResults = InCompilationResults; }
 
 	ERHIFeatureLevel::Type GetFeatureLevel() const { return FeatureLevel; }
 
@@ -465,7 +466,7 @@ protected:
 #endif // WITH_EDITOR
 
 private:
-	TArray<FString> CompileOutputMessages;
+	FComputeKernelCompileResults CompilationResults;
 
 	/** 
 	 * Game thread tracked shader map, which is ref counted and manages shader map lifetime. 
diff --git a/Engine/Plugins/Runtime/ComputeFramework/Source/ComputeFramework/Private/ComputeGraph.cpp b/Engine/Plugins/Runtime/ComputeFramework/Source/ComputeFramework/Private/ComputeGraph.cpp
--- a/Engine/Plugins/Runtime/ComputeFramework/Source/ComputeFramework/Private/ComputeGraph.cpp
+++ b/Engine/Plugins/Runtime/ComputeFramework/Source/ComputeFramework/Private/ComputeGraph.cpp
@@ -579,7 +579,7 @@ void UComputeGraph::CacheResourceShadersForRendering(uint32 CompilationFlags)
 			TUniquePtr<FShaderParametersMetadataAllocations> ShaderParameterMetadataAllocations = MakeUnique<FShaderParametersMetadataAllocations>();
 
 			FString ShaderEntryPoint = Kernel->KernelSource->EntryPoint;
-			FString ShaderFriendlyName = GetOuter()->GetName() + TEXT("_") + ShaderEntryPoint;
+			FString ShaderFriendlyName = GetOuter()->GetName() / GetFName().GetPlainNameString() / ShaderEntryPoint;
 			FString ShaderSource = BuildKernelSource(KernelIndex, *Kernel->KernelSource, AdditionalSources, ShaderHashKey, *ShaderDefinitionSet, *ShaderPermutationVector);
 			FShaderParametersMetadata* ShaderParameterMetadata = BuildKernelShaderMetadata(KernelIndex, *ShaderParameterMetadataAllocations);
 
@@ -649,11 +649,6 @@ void UComputeGraph::CacheShadersForResource(
 			*KernelResource->GetFriendlyName(),
 			*LegacyShaderPlatformToShaderFormat(ShaderPlatform).ToString()
 		);
-
-		for (FString const& Message : KernelResource->GetCompileOutputMessages())
-		{
-			UE_LOG(LogComputeFramework, Warning, TEXT("   %s"), *Message);
-		}
 	}
 }
 
@@ -664,7 +659,7 @@ void UComputeGraph::ShaderCompileCompletionCallback(FComputeKernelResource const
 	{
 		if (KernelResource == KernelResources[KernelIndex].Get())
 		{
-			OnKernelCompilationComplete(KernelIndex, KernelResource->GetCompileOutputMessages());
+			OnKernelCompilationComplete(KernelIndex, KernelResource->GetCompilationResults());
 		}
 	}
 }
diff --git a/Engine/Plugins/Runtime/ComputeFramework/Source/ComputeFramework/Private/ComputeKernelShaderCompilationManager.cpp b/Engine/Plugins/Runtime/ComputeFramework/Source/ComputeFramework/Private/ComputeKernelShaderCompilationManager.cpp
--- a/Engine/Plugins/Runtime/ComputeFramework/Source/ComputeFramework/Private/ComputeKernelShaderCompilationManager.cpp
+++ b/Engine/Plugins/Runtime/ComputeFramework/Source/ComputeFramework/Private/ComputeKernelShaderCompilationManager.cpp
@@ -14,7 +14,7 @@
 #include "Interfaces/ITargetPlatformManagerModule.h"
 #endif
 
-DEFINE_LOG_CATEGORY_STATIC(LogComputeKernelShaderCompiler, All, All);
+DEFINE_LOG_CATEGORY_STATIC(LogComputeKernelShaderCompiler, Log, All);
 
 static int32 GShowComputeKernelShaderWarnings = 0;
 static FAutoConsoleVariableRef CVarShowComputeKernelShaderWarnings(
@@ -75,11 +75,11 @@ void FComputeKernelShaderCompilationManager::ProcessAsyncResults()
 		CurrentJob->bSucceeded = CurrentJob->Output.bSucceeded;
 		if (CurrentJob->Output.bSucceeded)
 		{
-			UE_LOG(LogComputeKernelShaderCompiler, Log, TEXT("GPU shader compile succeeded. Id %d"), CurrentJob->Id);
+			UE_LOG(LogComputeKernelShaderCompiler, Verbose, TEXT("GPU shader compile succeeded. Id %d"), CurrentJob->Id);
 		}
 		else
 		{
-			UE_LOG(LogComputeKernelShaderCompiler, Log, TEXT("ERROR: GPU shader compile failed! Id %d"), CurrentJob->Id);
+			UE_LOG(LogComputeKernelShaderCompiler, Verbose, TEXT("GPU shader compile failed! Id %d"), CurrentJob->Id);
 		}
 
 		FComputeKernelShaderMapCompileResults& ShaderMapResults = ComputeKernelShaderMapJobs.FindChecked(CurrentJob->Id);
@@ -105,6 +105,95 @@ void FComputeKernelShaderCompilationManager::ProcessAsyncResults()
 	}
 }
 
+static bool ParseShaderCompilerError(FShaderCompilerError const& InError, FComputeKernelCompileMessage& OutMessage)
+{
+	FShaderCompilerError Error = InError;
+	Error.ExtractSourceLocation();
+
+	// We ignore error messages that don't have a line information.
+	FString Line, Column;
+	if (!Error.HasLineMarker())
+	{
+		return false;
+	}
+	if (!Error.ErrorLineString.Split(TEXT(","), &Line, &Column))
+	{
+		return false;
+	}
+	if (!Line.IsNumeric() || !Column.IsNumeric())
+	{
+		return false;
+	}
+
+	if (Error.StrippedErrorMessage.RemoveFromStart(TEXT("error: ")))
+	{
+		OutMessage.Type = FComputeKernelCompileMessage::EMessageType::Error;
+	}
+	else if (Error.StrippedErrorMessage.RemoveFromStart(TEXT("warning: ")))
+	{
+		OutMessage.Type = FComputeKernelCompileMessage::EMessageType::Warning;
+	}
+	else if (Error.StrippedErrorMessage.RemoveFromStart(TEXT("note: ")))
+	{
+		OutMessage.Type = FComputeKernelCompileMessage::EMessageType::Info;
+	}
+
+	OutMessage.Text = Error.StrippedErrorMessage;
+	OutMessage.VirtualFilePath = Error.ErrorVirtualFilePath;
+
+	if (OutMessage.VirtualFilePath.StartsWith(TEXT("/")) && !OutMessage.VirtualFilePath.StartsWith(TEXT("/UObject/")) && !OutMessage.VirtualFilePath.Contains(TEXT("/Generated/")))
+	{
+		OutMessage.RealFilePath = Error.GetShaderSourceFilePath();
+	}
+
+	LexFromString(OutMessage.Line, *Line);
+	LexFromString(OutMessage.ColumnStart, *Column);
+	OutMessage.ColumnEnd = OutMessage.ColumnStart;
+	
+	for (TCHAR Character : Error.HighlightedLineMarker)
+	{
+		OutMessage.ColumnEnd += (Character == TEXT('~')) ? 1 : 0;
+	}
+
+	return true;
+}
+
+static void LogShaderCompilerErrors(FComputeKernelCompileResults const& Results)
+{
+	for (FComputeKernelCompileMessage const& Message : Results.Messages)
+	{
+		FString Path = Message.RealFilePath.IsEmpty() ? Message.VirtualFilePath : Message.RealFilePath;
+
+		const bool bPreparePathForVisualStudioHotlink = !Message.RealFilePath.IsEmpty() && FPlatformMisc::IsDebuggerPresent();
+		if (bPreparePathForVisualStudioHotlink)
+		{
+			// Convert path to absolute, and prepend with newline so that it is clickable in Visual Studio.
+			Path = TEXT("\n") + IFileManager::Get().ConvertToAbsolutePathForExternalAppForRead(*Message.RealFilePath);
+		}
+
+		FString Line;
+		if (Message.ColumnStart == Message.ColumnEnd)
+		{
+			Line = FString::Printf(TEXT("(%d,%d)"), Message.Line, Message.ColumnStart);
+		}
+		else
+		{
+			Line = FString::Printf(TEXT("(%d,%d-%d)"), Message.Line, Message.ColumnStart, Message.ColumnEnd);
+		}
+
+		FString MessageText = FString::Printf(TEXT("%s%s: %s"), *Path, *Line, *Message.Text);
+
+		if (Message.Type == FComputeKernelCompileMessage::EMessageType::Warning)
+		{
+			UE_LOG(LogComputeKernelShaderCompiler, Warning, TEXT("%s"), *MessageText);
+		}
+		else if (Message.Type == FComputeKernelCompileMessage::EMessageType::Error)
+		{
+			UE_LOG(LogComputeKernelShaderCompiler, Error, TEXT("%s"), *MessageText);
+		}
+	}
+}
+
 void FComputeKernelShaderCompilationManager::ProcessCompiledComputeKernelShaderMaps(
 	TMap<int32, FComputeKernelShaderMapFinalizeResults>& CompiledShaderMaps,
 	float TimeBudget)
@@ -132,9 +221,9 @@ void FComputeKernelShaderCompilationManager::ProcessCompiledComputeKernelShaderM
 
 		if (ShaderMap && Kernels)
 		{
-			TArray<FString> Errors;
 			FComputeKernelShaderMapFinalizeResults& CompileResults = ProcessIt.Value();
 			TArray<FShaderCommonCompileJobPtr>& ResultArray = CompileResults.FinishedJobs;
+			FComputeKernelCompileResults ProcessedCompileResults;
 
 			// Make a copy of the array as this entry of FComputeKernelShaderMap::ShaderMapsBeingCompiled will be removed below
 			TArray<FComputeKernelResource*> KernelsArray = *Kernels;
@@ -154,21 +243,21 @@ void FComputeKernelShaderCompilationManager::ProcessCompiledComputeKernelShaderM
 				{
 					for (int32 ErrorIndex = 0; ErrorIndex < CurrentJob.Output.Errors.Num(); ErrorIndex++)
 					{
-						Errors.AddUnique(CurrentJob.Output.Errors[ErrorIndex].GetErrorString());
+						FComputeKernelCompileMessage Message;
+						if (ParseShaderCompilerError(CurrentJob.Output.Errors[ErrorIndex], Message))
+						{
+							ProcessedCompileResults.Messages.AddUnique(Message);
+						}
 					}
 
-					if (CurrentJob.Output.Errors.Num())
+					if (ProcessedCompileResults.Messages.Num())
 					{
-						UE_LOG(LogComputeKernelShaderCompiler, Log, TEXT("There were errors for job \"%s\""), *CurrentJob.Input.DebugGroupName)
-						for (const FShaderCompilerError& Error : CurrentJob.Output.Errors)
-						{
-							UE_LOG(LogComputeKernelShaderCompiler, Log, TEXT("Error: %s"), *Error.GetErrorString())
-						}
+						UE_LOG(LogComputeKernelShaderCompiler, Verbose, TEXT("There were errors for job \"%s\""), *CurrentJob.Input.DebugGroupName)
 					}
 				}
 				else
 				{
-					UE_LOG(LogComputeKernelShaderCompiler, Log, TEXT("There were NO errors for job \"%s\""), *CurrentJob.Input.DebugGroupName);
+					UE_LOG(LogComputeKernelShaderCompiler, Verbose, TEXT("There were NO errors for job \"%s\""), *CurrentJob.Input.DebugGroupName);
 				}
 			}
 
@@ -197,34 +286,12 @@ void FComputeKernelShaderCompilationManager::ProcessCompiledComputeKernelShaderM
 					// This avoids applying shadermaps which are out of date and a newer one is in the async compiling pipeline
 					if (Kernel->IsSame(CompletedShaderMap->GetShaderMapId()))
 					{
-						if (Errors.Num() != 0)
-						{
-							FString SourceCode;
-							Kernel->GetHLSLSource(SourceCode);
-							UE_LOG(LogComputeKernelShaderCompiler, Log, TEXT("Compile output as text:"));
-							UE_LOG(LogComputeKernelShaderCompiler, Log, TEXT("==================================================================================="));
-							TArray<FString> OutputByLines;
-							SourceCode.ParseIntoArrayLines(OutputByLines, false);
-							for (int32 i = 0; i < OutputByLines.Num(); i++)
-							{
-								UE_LOG(LogComputeKernelShaderCompiler, Log, TEXT("/*%04d*/\t\t%s"), i + 1, *OutputByLines[i]);
-							}
-							UE_LOG(LogComputeKernelShaderCompiler, Log, TEXT("==================================================================================="));
-						}
-
 						if (!bSuccess)
 						{
 							// Propagate error messages
-							Kernel->SetCompileOutputMessages(Errors);
+							LogShaderCompilerErrors(ProcessedCompileResults);
+							Kernel->SetCompilationResults(ProcessedCompileResults);
 							KernelsToUpdate.Add(Kernel, nullptr);
-
-							for (int32 ErrorIndex = 0; ErrorIndex < Errors.Num(); ErrorIndex++)
-							{
-								FString ErrorMessage = Errors[ErrorIndex];
-								// Work around build machine string matching heuristics that will cause a cook to fail
-								ErrorMessage.ReplaceInline(TEXT("error "), TEXT("err0r "), ESearchCase::CaseSensitive);
-								UE_LOG(LogComputeKernelShaderCompiler, Warning, TEXT("	%s"), *ErrorMessage);
-							}
 						}
 						else
 						{
@@ -236,17 +303,14 @@ void FComputeKernelShaderCompilationManager::ProcessCompiledComputeKernelShaderM
 								KernelsToUpdate.Add(Kernel, CompletedShaderMap);
 							}
 
-							if (GShowComputeKernelShaderWarnings && Errors.Num() > 0)
+							if (GShowComputeKernelShaderWarnings && ProcessedCompileResults.Messages.Num() > 0)
 							{
 								UE_LOG(LogComputeKernelShaderCompiler, Warning, TEXT("Warnings while compiling ComputeKernel %s for platform %s:"),
 									*Kernel->GetFriendlyName(),
 									*LegacyShaderPlatformToShaderFormat(ShaderMap->GetShaderPlatform()).ToString());
-								for (int32 ErrorIndex = 0; ErrorIndex < Errors.Num(); ErrorIndex++)
-								{
-									UE_LOG(LogComputeKernelShaderCompiler, Warning, TEXT("	%s"), *Errors[ErrorIndex]);
-								}
 
-								Kernel->SetCompileOutputMessages(Errors);
+								LogShaderCompilerErrors(ProcessedCompileResults);
+								Kernel->SetCompilationResults(ProcessedCompileResults);
 							}
 						}
 					}
diff --git a/Engine/Plugins/Runtime/ComputeFramework/Source/ComputeFramework/Private/ComputeKernelShaderMap.cpp b/Engine/Plugins/Runtime/ComputeFramework/Source/ComputeFramework/Private/ComputeKernelShaderMap.cpp
--- a/Engine/Plugins/Runtime/ComputeFramework/Source/ComputeFramework/Private/ComputeKernelShaderMap.cpp
+++ b/Engine/Plugins/Runtime/ComputeFramework/Source/ComputeFramework/Private/ComputeKernelShaderMap.cpp
@@ -334,7 +334,7 @@ void FComputeKernelShaderMap::Compile(
 		{
 			// Assign a unique identifier so that shaders from this shader map can be associated with it after a deferred compile
 			CompilingId = FShaderCommonCompileJob::GetNextJobId();
-			UE_LOG(LogComputeFramework, Log, TEXT("CompilingId = %p %d"), InKernel, CompilingId);
+			UE_LOG(LogComputeFramework, Verbose, TEXT("CompilingId = %p %d"), InKernel, CompilingId);
 			InKernel->AddCompileId(CompilingId);
 
 			// Setup the compilation environment.
@@ -394,7 +394,7 @@ void FComputeKernelShaderMap::Compile(
   
 			if (!Kernel)
 			{
-				UE_LOG(LogComputeFramework, Log, TEXT("		%u Shaders"), NumShaders);
+				UE_LOG(LogComputeFramework, Verbose, TEXT("		%u Shaders"), NumShaders);
 			}
 
 			// Register this shader map in the global ComputeKernel->shadermap map
@@ -491,7 +491,7 @@ bool FComputeKernelShaderMap::TryToAddToExistingCompilationTask(FComputeKernelRe
 	{
 		CorrespondingKernels->AddUnique(InKernel);
 
-		UE_LOG(LogComputeFramework, Log, TEXT("TryToAddToExistingCompilationTask %p %d"), InKernel, GetCompilingId());
+		UE_LOG(LogComputeFramework, Verbose, TEXT("TryToAddToExistingCompilationTask %p %d"), InKernel, GetCompilingId());
 
 #if DEBUG_INFINITESHADERCOMPILE
 		UE_LOG(LogTemp, Display, TEXT("Added shader map 0x%08X%08X from ComputeKernel 0x%08X%08X"), (int)((int64)(this) >> 32), (int)((int64)(this)), (int)((int64)(InKernel) >> 32), (int)((int64)(InKernel)));
diff --git a/Engine/Plugins/Runtime/ComputeFramework/Source/ComputeFramework/Private/ComputeKernelShared.cpp b/Engine/Plugins/Runtime/ComputeFramework/Source/ComputeFramework/Private/ComputeKernelShared.cpp
--- a/Engine/Plugins/Runtime/ComputeFramework/Source/ComputeFramework/Private/ComputeKernelShared.cpp
+++ b/Engine/Plugins/Runtime/ComputeFramework/Source/ComputeFramework/Private/ComputeKernelShared.cpp
@@ -60,8 +60,11 @@ bool FComputeKernelResource::ShouldCache(EShaderPlatform InPlatform, const FShad
 
 void FComputeKernelResource::NotifyCompilationFinished(FString const& ResultMessage)
 {
-	UE_LOG(LogComputeFramework, Display, TEXT("%s"), *ResultMessage);
-	CompileOutputMessages.Add(ResultMessage);
+	UE_LOG(LogComputeFramework, Log, TEXT("%s"), *ResultMessage);
+	FComputeKernelCompileMessage Message;
+	Message.Type = FComputeKernelCompileMessage::EMessageType::Info;
+	Message.Text = ResultMessage;
+	CompilationResults.Messages.Add(Message);
 	OnCompilationCompleteDelegate.ExecuteIfBound(this);
 }
 
@@ -83,7 +86,7 @@ void FComputeKernelResource::RemoveOutstandingCompileId(const int32 InOldOutstan
 {
 	if (0 <= OutstandingCompileShaderMapIds.Remove(InOldOutstandingCompileShaderMapId))
 	{
-		UE_LOG(LogComputeFramework, Log, TEXT("RemoveOutstandingCompileId %p %d"), this, InOldOutstandingCompileShaderMapId);
+		UE_LOG(LogComputeFramework, Verbose, TEXT("RemoveOutstandingCompileId %p %d"), this, InOldOutstandingCompileShaderMapId);
 	}
 }
 
@@ -244,7 +247,7 @@ void FComputeKernelResource::SetupResource(
 	ShaderPermutationVector = InShaderPermutationVector;
 	ShaderParameterMetadataAllocations = MoveTemp(InShaderParameterMetadataAllocations);
 	ShaderParameterMetadata = InShaderParameterMetadata;
-	CompileOutputMessages.Reset();
+	CompilationResults.Messages.Reset();
 #if WITH_EDITOR
 	AssetPath = InAssetPath;
 #endif
@@ -334,12 +337,7 @@ bool FComputeKernelResource::CacheShaders(const FComputeKernelShaderMapId& InSha
 			{
 #if WITH_EDITOR
 				FComputeKernelShaderMap::LoadFromDerivedDataCache(this, InShaderMapId, InPlatform, GameThreadShaderMap);
-				if (GameThreadShaderMap && GameThreadShaderMap->IsValid())
-				{
-					FString Message = FString::Printf(TEXT("Loaded shaders for %s from DDC."), *GetFriendlyName());
-					CompileOutputMessages.Add(Message);
-				}
-				else
+				if (GameThreadShaderMap == nullptr || !GameThreadShaderMap->IsValid())
 				{
 					UE_LOG(LogComputeFramework, Log, TEXT("Loading shader for kernel %s from DDC failed. Shader needs recompile."), *GetFriendlyName());
 				}
@@ -360,7 +358,7 @@ bool FComputeKernelResource::CacheShaders(const FComputeKernelShaderMapId& InSha
 		UE_LOG(LogComputeFramework, Display, TEXT("Found existing compiling shader for kernel %s, linking to other GameThreadShaderMap 0x%08X%08X"), *GetFriendlyName(), (int)((int64)(GameThreadShaderMap.GetReference()) >> 32), (int)((int64)(GameThreadShaderMap.GetReference())));
 #endif
 		OutstandingCompileShaderMapIds.AddUnique(GameThreadShaderMap->GetCompilingId());
-		UE_LOG(LogComputeFramework, Log, TEXT("CacheShaders AddUniqueExisting %p %d"), this, GameThreadShaderMap->GetCompilingId());
+		UE_LOG(LogComputeFramework, Verbose, TEXT("CacheShaders AddUniqueExisting %p %d"), this, GameThreadShaderMap->GetCompilingId());
 
 		GameThreadShaderMap = nullptr;
 		bSucceeded = true;
@@ -421,7 +419,7 @@ void FComputeKernelResource::FinishCompilation()
 	{
 		for (int32 i = 0; i < ShaderMapIdsToFinish.Num(); i++)
 		{
-			UE_LOG(LogComputeFramework, Log, TEXT("FinishCompilation()[%d] %s id %d!"), i, *GetFriendlyName(), ShaderMapIdsToFinish[i]);
+			UE_LOG(LogComputeFramework, Verbose, TEXT("FinishCompilation()[%d] %s id %d!"), i, *GetFriendlyName(), ShaderMapIdsToFinish[i]);
 		}
 	
 		// Block until the shader maps that we will save have finished being compiled
@@ -500,7 +498,7 @@ bool FComputeKernelResource::BeginCompileShaderMap(const FComputeKernelShaderMap
 	}
 	else
 	{
-		UE_LOG(LogComputeFramework, Log, TEXT("BeginCompileShaderMap AddUnique %p %d"), this, NewShaderMap->GetCompilingId());
+		UE_LOG(LogComputeFramework, Verbose, TEXT("BeginCompileShaderMap AddUnique %p %d"), this, NewShaderMap->GetCompilingId());
 		OutstandingCompileShaderMapIds.AddUnique(NewShaderMap->GetCompilingId());
 		
 		// Async compile, use nullptr to detect it if used
diff --git a/Engine/Plugins/Runtime/ComputeFramework/Source/ComputeFramework/Public/ComputeFramework/ComputeGraph.h b/Engine/Plugins/Runtime/ComputeFramework/Source/ComputeFramework/Public/ComputeFramework/ComputeGraph.h
--- a/Engine/Plugins/Runtime/ComputeFramework/Source/ComputeFramework/Public/ComputeFramework/ComputeGraph.h
+++ b/Engine/Plugins/Runtime/ComputeFramework/Source/ComputeFramework/Public/ComputeFramework/ComputeGraph.h
@@ -12,11 +12,12 @@
 
 #include "ComputeGraph.generated.h"
 
+enum EShaderPlatform : uint16;
 class FArchive;
+struct FComputeKernelCompileResults;
 struct FComputeKernelDefinitionSet;
 struct FComputeKernelPermutationVector;
 class FComputeKernelResource;
-class UComputeKernelSource;
 class FComputeGraphRenderProxy;
 class ITargetPlatform;
 class FShaderParametersMetadata;
@@ -24,7 +25,7 @@ struct FShaderParametersMetadataAllocations;
 class UComputeDataInterface;
 class UComputeDataProvider;
 class UComputeKernel;
-enum EShaderPlatform : uint16;
+class UComputeKernelSource;
 
 /** 
  * Description of a single edge in a UComputeGraph. 
@@ -102,7 +103,7 @@ public:
 	virtual ~UComputeGraph();
 
 	/** Called each time that a single kernel shader compilation is completed. */
-	virtual void OnKernelCompilationComplete(int32 InKernelIndex, const TArray<FString>& InCompileOutputMessages) {}
+	virtual void OnKernelCompilationComplete(int32 InKernelIndex, FComputeKernelCompileResults const& InCompileResults) {}
 
 	/** 
 	 * Returns true if graph is valid. 
diff --git a/Engine/Plugins/Runtime/ComputeFramework/Source/ComputeFramework/Public/ComputeFramework/ComputeKernelCompileResult.h b/Engine/Plugins/Runtime/ComputeFramework/Source/ComputeFramework/Public/ComputeFramework/ComputeKernelCompileResult.h
new file mode 100644
--- /dev/null
+++ b/Engine/Plugins/Runtime/ComputeFramework/Source/ComputeFramework/Public/ComputeFramework/ComputeKernelCompileResult.h
@@ -0,0 +1,37 @@
+// Copyright Epic Games, Inc. All Rights Reserved.
+
+#pragma once
+
+#include "CoreMinimal.h"
+
+/** Struct containing a message returned by shader compilation. */
+struct FComputeKernelCompileMessage
+{
+	enum class EMessageType
+	{
+		None,
+		Info,
+		Warning,
+		Error
+	};
+
+	EMessageType Type = EMessageType::None;
+	FString Text;
+	FString VirtualFilePath;
+	FString RealFilePath;
+	int32 Line = -1;
+	int32 ColumnStart = -1;
+	int32 ColumnEnd = -1;
+	
+	/** Equality operator used to deduplicate messages. */
+	bool operator==(FComputeKernelCompileMessage const& Rhs) const
+	{
+		return Type == Rhs.Type && Text == Rhs.Text && VirtualFilePath == Rhs.VirtualFilePath && RealFilePath == Rhs.RealFilePath && Line == Rhs.Line && ColumnStart == Rhs.ColumnStart && ColumnEnd == Rhs.ColumnEnd;
+	}
+};
+
+/** Struct containing all of the messages returned by shader compilation. */
+struct FComputeKernelCompileResults
+{
+	TArray<FComputeKernelCompileMessage> Messages;
+};
