ComputeFramework: Put DataInterface code inside virtual generated files for shader compilation.
This allows us to jump to compile errors in source for data interfaces that have simple source templates.

diff --git a/Engine/Plugins/Runtime/ComputeFramework/Source/ComputeFramework/Private/ComputeFramework/ComputeKernelShared.h b/Engine/Plugins/Runtime/ComputeFramework/Source/ComputeFramework/Private/ComputeFramework/ComputeKernelShared.h
--- a/Engine/Plugins/Runtime/ComputeFramework/Source/ComputeFramework/Private/ComputeFramework/ComputeKernelShared.h
+++ b/Engine/Plugins/Runtime/ComputeFramework/Source/ComputeFramework/Private/ComputeFramework/ComputeKernelShared.h
@@ -432,6 +432,7 @@ public:
 		FString const& InShaderHashKey,
 		FString const& InShaderSource,
 		TMap<FString, FString> const& InAdditionalSources,
+		TMap<FString, FString> const& InGeneratedSources,
 		TSharedPtr<FComputeKernelDefinitionSet>& InShaderDefinitionSet,
 		TSharedPtr<FComputeKernelPermutationVector>& InShaderPermutationVector,
 		TUniquePtr<FShaderParametersMetadataAllocations>& InShaderParameterMetadataAllocations,
@@ -496,6 +497,8 @@ private:
 
 	/** Additional source code. Stored as a map from a virtual source file name to the shader source.  */
 	TMap<FString, FString> AdditionalSources;
+	/** Generated source code. Stored as a map from a virtual source file name to the shader source.  */
+	TMap<FString, FString> GeneratedSources;
 
 	/** Defines used when compiling shaders. Object can be shared across a number of shader formats during cook. */
 	TSharedPtr<FComputeKernelDefinitionSet> ShaderDefinitionSet;
diff --git a/Engine/Plugins/Runtime/ComputeFramework/Source/ComputeFramework/Private/ComputeGraph.cpp b/Engine/Plugins/Runtime/ComputeFramework/Source/ComputeFramework/Private/ComputeGraph.cpp
--- a/Engine/Plugins/Runtime/ComputeFramework/Source/ComputeFramework/Private/ComputeGraph.cpp
+++ b/Engine/Plugins/Runtime/ComputeFramework/Source/ComputeFramework/Private/ComputeGraph.cpp
@@ -439,7 +439,8 @@ FString UComputeGraph::BuildKernelSource(
 	UComputeKernelSource const& InKernelSource,
 	TMap<FString, FString> const& InAdditionalSources,
 	FString& OutHashKey,
-	FComputeKernelDefinitionSet& OutDefinitionSet, 
+	TMap<FString, FString>& OutGeneratedSources,
+	FComputeKernelDefinitionSet& OutDefinitionSet,
 	FComputeKernelPermutationVector& OutPermutationVector) const
 {
 	FString HLSL;
@@ -481,7 +482,21 @@ FString UComputeGraph::BuildKernelSource(
 		{
 			// Add a unique prefix to generate unique names in the data interface shader code.
 			FString NamePrefix = GetUniqueDataInterfaceName(DataInterface, DataProviderIndex);
-			DataInterface->GetHLSL(HLSL, NamePrefix);
+
+			// Data interface optionally put source in a generated file that maps to an on-disk virtual path.
+			if (TCHAR const* ShaderVirtualPath = DataInterface->GetShaderVirtualPath())
+			{
+				// The generated path has a magic unique prefix which the compilation manager knows to strip before resolving errors.
+				FString MagicVirtualPath = FString::Printf(TEXT("/Engine/Generated/DataInterface/%s%s"), *NamePrefix, ShaderVirtualPath);
+				HLSL += FString::Printf(TEXT("#include \"%s\"\n"), *MagicVirtualPath);
+				FString DataInterfaceHLSL;
+				DataInterface->GetHLSL(DataInterfaceHLSL, NamePrefix);
+				OutGeneratedSources.Add(MagicVirtualPath, DataInterfaceHLSL);
+			}
+			else
+			{
+				DataInterface->GetHLSL(HLSL, NamePrefix);
+			}
 
 			DataInterface->GetStructDeclarations(StructsSeen, StructDeclarations);
 			
@@ -574,13 +589,14 @@ void UComputeGraph::CacheResourceShadersForRendering(uint32 CompilationFlags)
 			TMap<FString, FString> AdditionalSources = GatherAdditionalSources(Kernel->KernelSource->AdditionalSources);
 
 			FString ShaderHashKey;
+			TMap<FString, FString> GeneratedSources;
 			TSharedPtr<FComputeKernelDefinitionSet> ShaderDefinitionSet = MakeShared<FComputeKernelDefinitionSet>();
 			TSharedPtr<FComputeKernelPermutationVector> ShaderPermutationVector = MakeShared<FComputeKernelPermutationVector>();
 			TUniquePtr<FShaderParametersMetadataAllocations> ShaderParameterMetadataAllocations = MakeUnique<FShaderParametersMetadataAllocations>();
 
 			FString ShaderEntryPoint = Kernel->KernelSource->EntryPoint;
 			FString ShaderFriendlyName = GetOuter()->GetName() / GetFName().GetPlainNameString() / ShaderEntryPoint;
-			FString ShaderSource = BuildKernelSource(KernelIndex, *Kernel->KernelSource, AdditionalSources, ShaderHashKey, *ShaderDefinitionSet, *ShaderPermutationVector);
+			FString ShaderSource = BuildKernelSource(KernelIndex, *Kernel->KernelSource, AdditionalSources, ShaderHashKey, GeneratedSources, *ShaderDefinitionSet, *ShaderPermutationVector);
 			FShaderParametersMetadata* ShaderParameterMetadata = BuildKernelShaderMetadata(KernelIndex, *ShaderParameterMetadataAllocations);
 
 			const ERHIFeatureLevel::Type CacheFeatureLevel = GMaxRHIFeatureLevel;
@@ -595,6 +611,7 @@ void UComputeGraph::CacheResourceShadersForRendering(uint32 CompilationFlags)
 				ShaderHashKey,
 				ShaderSource,
 				AdditionalSources,
+				GeneratedSources,
 				ShaderDefinitionSet,
 				ShaderPermutationVector,
 				ShaderParameterMetadataAllocations,
@@ -685,12 +702,13 @@ void UComputeGraph::BeginCacheForCookedPlatformData(ITargetPlatform const* Targe
 			TMap<FString, FString> AdditionalSources = GatherAdditionalSources(KernelInvocations[KernelIndex]->KernelSource->AdditionalSources);
 
 			FString ShaderHashKey;
+			TMap<FString, FString> GeneratedSources;
 			TSharedPtr<FComputeKernelDefinitionSet> ShaderDefinitionSet = MakeShared<FComputeKernelDefinitionSet>();
 			TSharedPtr<FComputeKernelPermutationVector> ShaderPermutationVector = MakeShared<FComputeKernelPermutationVector>();
 
 			FString ShaderEntryPoint = KernelSource->EntryPoint;
 			FString ShaderFriendlyName = GetOuter()->GetName() + TEXT("_") + ShaderEntryPoint;
-			FString ShaderSource = BuildKernelSource(KernelIndex, *KernelSource, AdditionalSources, ShaderHashKey, *ShaderDefinitionSet, *ShaderPermutationVector);
+			FString ShaderSource = BuildKernelSource(KernelIndex, *KernelSource, AdditionalSources, ShaderHashKey, GeneratedSources, *ShaderDefinitionSet, *ShaderPermutationVector);
 
 			for (int32 ShaderFormatIndex = 0; ShaderFormatIndex < ShaderFormats.Num(); ++ShaderFormatIndex)
 			{
@@ -708,6 +726,7 @@ void UComputeGraph::BeginCacheForCookedPlatformData(ITargetPlatform const* Targe
 					ShaderHashKey,
 					ShaderSource,
 					AdditionalSources,
+					GeneratedSources,
 					ShaderDefinitionSet,
 					ShaderPermutationVector,
 					ShaderParameterMetadataAllocations,
diff --git a/Engine/Plugins/Runtime/ComputeFramework/Source/ComputeFramework/Private/ComputeKernelShaderCompilationManager.cpp b/Engine/Plugins/Runtime/ComputeFramework/Source/ComputeFramework/Private/ComputeKernelShaderCompilationManager.cpp
--- a/Engine/Plugins/Runtime/ComputeFramework/Source/ComputeFramework/Private/ComputeKernelShaderCompilationManager.cpp
+++ b/Engine/Plugins/Runtime/ComputeFramework/Source/ComputeFramework/Private/ComputeKernelShaderCompilationManager.cpp
@@ -141,9 +141,16 @@ static bool ParseShaderCompilerError(FShaderCompilerError const& InError, FCompu
 	OutMessage.Text = Error.StrippedErrorMessage;
 	OutMessage.VirtualFilePath = Error.ErrorVirtualFilePath;
 
-	if (OutMessage.VirtualFilePath.StartsWith(TEXT("/")) && !OutMessage.VirtualFilePath.StartsWith(TEXT("/UObject/")) && !OutMessage.VirtualFilePath.Contains(TEXT("/Generated/")))
+	// Fix up the DataInterface generated file paths before any error reporting.
+	// Magic path structure is set in ComputeGraph compilation.
+	if (OutMessage.VirtualFilePath.RemoveFromStart(TEXT("/Engine/Generated/DataInterface/")))
 	{
-		OutMessage.RealFilePath = Error.GetShaderSourceFilePath();
+		OutMessage.VirtualFilePath.MidInline(OutMessage.VirtualFilePath.Find(TEXT("/")));
+	}
+	// Store any disk paths before error reporting. Can skip some known cases that won't have disk paths.
+	if (OutMessage.VirtualFilePath.StartsWith(TEXT("/")) && !OutMessage.VirtualFilePath.StartsWith(TEXT("/Engine/Generated/")))
+	{
+		OutMessage.RealFilePath = GetShaderSourceFilePath(OutMessage.VirtualFilePath);
 	}
 
 	LexFromString(OutMessage.Line, *Line);
diff --git a/Engine/Plugins/Runtime/ComputeFramework/Source/ComputeFramework/Private/ComputeKernelShared.cpp b/Engine/Plugins/Runtime/ComputeFramework/Source/ComputeFramework/Private/ComputeKernelShared.cpp
--- a/Engine/Plugins/Runtime/ComputeFramework/Source/ComputeFramework/Private/ComputeKernelShared.cpp
+++ b/Engine/Plugins/Runtime/ComputeFramework/Source/ComputeFramework/Private/ComputeKernelShared.cpp
@@ -49,6 +49,7 @@ void FComputeKernelResource::SetupShaderCompilationEnvironment(EShaderPlatform I
 	}
 
 	OutEnvironment.IncludeVirtualPathToContentsMap.Append(AdditionalSources);
+	OutEnvironment.IncludeVirtualPathToContentsMap.Append(GeneratedSources);
 }
 #endif // WITH_EDITOR
 
@@ -230,6 +231,7 @@ void FComputeKernelResource::SetupResource(
 	FString const& InShaderHashKey,
 	FString const& InShaderSource,
 	TMap<FString, FString> const& InAdditionalSources,
+	TMap<FString, FString> const& InGeneratedSources,
 	TSharedPtr<FComputeKernelDefinitionSet>& InShaderDefinitionSet,
 	TSharedPtr<FComputeKernelPermutationVector>& InShaderPermutationVector,
 	TUniquePtr<FShaderParametersMetadataAllocations>& InShaderParameterMetadataAllocations,
@@ -243,6 +245,7 @@ void FComputeKernelResource::SetupResource(
 	ShaderCodeHash = GetTypeHash(InShaderHashKey);
 	ShaderSource = InShaderSource;
 	AdditionalSources = InAdditionalSources;
+	GeneratedSources = InGeneratedSources;
 	ShaderDefinitionSet = InShaderDefinitionSet;
 	ShaderPermutationVector = InShaderPermutationVector;
 	ShaderParameterMetadataAllocations = MoveTemp(InShaderParameterMetadataAllocations);
diff --git a/Engine/Plugins/Runtime/ComputeFramework/Source/ComputeFramework/Public/ComputeFramework/ComputeDataInterface.h b/Engine/Plugins/Runtime/ComputeFramework/Source/ComputeFramework/Public/ComputeFramework/ComputeDataInterface.h
--- a/Engine/Plugins/Runtime/ComputeFramework/Source/ComputeFramework/Public/ComputeFramework/ComputeDataInterface.h
+++ b/Engine/Plugins/Runtime/ComputeFramework/Source/ComputeFramework/Public/ComputeFramework/ComputeDataInterface.h
@@ -37,6 +37,8 @@ public:
 	virtual void GetSupportedOutputs(TArray<FShaderFunctionDefinition>& OutFunctions) const {}
 	/** Gather the shader metadata exposed by the data provider payload. */
 	virtual void GetShaderParameters(TCHAR const* UID, FShaderParametersMetadataBuilder& InOutBuilder, FShaderParametersMetadataAllocations& InOutAllocations) const {}
+	/** If the data interface HLSL comes from a file this should fill the virtual file path. This is used to resolve shader compilation errors. */
+	virtual TCHAR const* GetShaderVirtualPath() const { return nullptr; }
 	/** Get a hash that changes on any data interface changes that affect kernel compilation. */
 	virtual void GetShaderHash(FString& InOutKey) const {}
 	/** Gather any extra struct types that this data provider relies on. */
diff --git a/Engine/Plugins/Runtime/ComputeFramework/Source/ComputeFramework/Public/ComputeFramework/ComputeGraph.h b/Engine/Plugins/Runtime/ComputeFramework/Source/ComputeFramework/Public/ComputeFramework/ComputeGraph.h
--- a/Engine/Plugins/Runtime/ComputeFramework/Source/ComputeFramework/Public/ComputeFramework/ComputeGraph.h
+++ b/Engine/Plugins/Runtime/ComputeFramework/Source/ComputeFramework/Public/ComputeFramework/ComputeGraph.h
@@ -160,7 +160,8 @@ private:
 		UComputeKernelSource const& InKernelSource,
 		TMap<FString, FString> const& InAdditionalSources,
 		FString& OutHashKey,
-		FComputeKernelDefinitionSet& OutDefinitionSet, 
+		TMap<FString, FString>& OutGeneratedSources,
+		FComputeKernelDefinitionSet& OutDefinitionSet,
 		FComputeKernelPermutationVector& OutPermutationVector) const;
 
 	/** Cache shader resources for all kernels in the graph. */
