ComputeFramework: Support for unified dispatches.
When a kernel supports unified dispatch it combines any subinvocations into a single dispatch invocation.
This can only happen if all data interfaces support unified dispatch, and all subinvocation shaders are the same.
Some data interfaces, such as skeleton, don't support this because each section keeps its own bone buffer to keep bone count to a minimum.
In future we should look to change the underlying structures to support unified dispatch as much as possible, since it will be more performant.

diff --git a/Engine/Plugins/Runtime/ComputeFramework/Source/ComputeFramework/Private/ComputeFramework/ComputeGraphRenderProxy.h b/Engine/Plugins/Runtime/ComputeFramework/Source/ComputeFramework/Private/ComputeFramework/ComputeGraphRenderProxy.h
--- a/Engine/Plugins/Runtime/ComputeFramework/Source/ComputeFramework/Private/ComputeFramework/ComputeGraphRenderProxy.h
+++ b/Engine/Plugins/Runtime/ComputeFramework/Source/ComputeFramework/Private/ComputeFramework/ComputeGraphRenderProxy.h
@@ -32,6 +32,8 @@ public:
 		TArray<int32> BoundProviderIndices;
 		/** The index of the special execution data provider in the full graph data provider array. */
 		int32 ExecutionProviderIndex = -1;
+		/** Whether the kernel can combine multiple supports combinined sub-invocations. */
+		bool bSupportsUnifiedDispatch = false;
 	};
 
 	/** Friendly name for the owner graph. */
diff --git a/Engine/Plugins/Runtime/ComputeFramework/Source/ComputeFramework/Private/ComputeGraph.cpp b/Engine/Plugins/Runtime/ComputeFramework/Source/ComputeFramework/Private/ComputeGraph.cpp
--- a/Engine/Plugins/Runtime/ComputeFramework/Source/ComputeFramework/Private/ComputeGraph.cpp
+++ b/Engine/Plugins/Runtime/ComputeFramework/Source/ComputeFramework/Private/ComputeGraph.cpp
@@ -313,6 +313,7 @@ FComputeGraphRenderProxy* UComputeGraph::CreateRenderProxy() const
 			Invocation.KernelGroupSize = Kernel->KernelSource->GroupSize;
 			Invocation.KernelResource = KernelResource;
 			Invocation.ShaderParameterMetadata = BuildKernelShaderMetadata(KernelIndex, *Proxy->ShaderParameterMetadataAllocations);
+			Invocation.bSupportsUnifiedDispatch = true;
 
 			for (FComputeGraphEdge const& GraphEdge : GraphEdges)
 			{
@@ -327,6 +328,8 @@ FComputeGraphRenderProxy* UComputeGraph::CreateRenderProxy() const
 						{
 							Invocation.ExecutionProviderIndex = GraphEdge.DataInterfaceIndex;
 						}
+						
+						Invocation.bSupportsUnifiedDispatch &= DataInterface->CanSupportUnifiedDispatch();
 					}
 				}
 			}
diff --git a/Engine/Plugins/Runtime/ComputeFramework/Source/ComputeFramework/Private/ComputeGraphWorker.cpp b/Engine/Plugins/Runtime/ComputeFramework/Source/ComputeFramework/Private/ComputeGraphWorker.cpp
--- a/Engine/Plugins/Runtime/ComputeFramework/Source/ComputeFramework/Private/ComputeGraphWorker.cpp
+++ b/Engine/Plugins/Runtime/ComputeFramework/Source/ComputeFramework/Private/ComputeGraphWorker.cpp
@@ -41,6 +41,10 @@ void FComputeGraphTaskWorker::SubmitWork(FRDGBuilder& GraphBuilder, FName InExec
 			FComputeGraphRenderProxy const* GraphRenderProxy = GraphInvocation.GraphRenderProxy;
 
 			TArray<TShaderRef<FComputeKernelShader>> Shaders;
+			TArray<int32> ShaderStartPerKernel;
+			ShaderStartPerKernel.Reserve(GraphRenderProxy->KernelInvocations.Num());
+			TArray<bool> UnifiedDispatchPerKernel;
+			UnifiedDispatchPerKernel.SetNumZeroed(GraphRenderProxy->KernelInvocations.Num());
 
 			// Validation phase.
 			// Check if all DataInterfaces are valid.
@@ -82,6 +86,7 @@ void FComputeGraphTaskWorker::SubmitWork(FRDGBuilder& GraphBuilder, FName InExec
 				}
 
 				// Get shader. This can fail if compilation is pending.
+				ShaderStartPerKernel.Add(Shaders.Num());
 				Shaders.Reserve(Shaders.Num() + NumSubInvocations);
 				for (int32 SubInvocationIndex = 0; bIsValid && SubInvocationIndex < NumSubInvocations; ++SubInvocationIndex)
 				{
@@ -89,6 +94,22 @@ void FComputeGraphTaskWorker::SubmitWork(FRDGBuilder& GraphBuilder, FName InExec
 					bIsValid &= Shader.IsValid();
 					Shaders.Add(Shader);
 				}
+
+				// Check if we can do unified dispatch and apply that if we can.
+				if (bIsValid && KernelInvocation.bSupportsUnifiedDispatch && NumSubInvocations > 1)
+				{
+					bool bSupportsUnifiedDispatch = true;
+					for (int32 SubInvocationIndex = 1; bSupportsUnifiedDispatch && SubInvocationIndex < NumSubInvocations; ++SubInvocationIndex)
+					{
+						bSupportsUnifiedDispatch &= Shaders[ShaderStartPerKernel[KernelIndex] + SubInvocationIndex] == Shaders[ShaderStartPerKernel[KernelIndex]];
+					}
+
+					if (bSupportsUnifiedDispatch)
+					{
+						UnifiedDispatchPerKernel[KernelIndex] = true;
+						Shaders.SetNum(Shaders.Num() - NumSubInvocations + 1);
+					}
+				}
 			}
 
 			// If we can't run the graph for any reason, back out now and apply fallback logic.
@@ -119,10 +140,21 @@ void FComputeGraphTaskWorker::SubmitWork(FRDGBuilder& GraphBuilder, FName InExec
 				RDG_EVENT_SCOPE(GraphBuilder, "%s", *KernelInvocation.KernelName);
 
 				TArray<FIntVector> ThreadCounts;
-				const int32 NumSubInvocations = GraphInvocation.DataProviderRenderProxies[KernelInvocation.ExecutionProviderIndex]->GetDispatchThreadCount(ThreadCounts);
+				int32 NumSubInvocations = GraphInvocation.DataProviderRenderProxies[KernelInvocation.ExecutionProviderIndex]->GetDispatchThreadCount(ThreadCounts);
+				
+				bool bIsUnifiedDispatch = UnifiedDispatchPerKernel[KernelIndex];
+				if (bIsUnifiedDispatch)
+				{
+					for (int32 SubInvocationIndex = 1; SubInvocationIndex < NumSubInvocations; ++SubInvocationIndex)
+					{
+						ThreadCounts[0] += ThreadCounts[SubInvocationIndex];
+					}
+					ThreadCounts.SetNum(1);
+					NumSubInvocations = 1;
+				}
 
 				TStridedView<FComputeKernelShader::FParameters> ParameterArray = GraphBuilder.AllocParameters<FComputeKernelShader::FParameters>(KernelInvocation.ShaderParameterMetadata, NumSubInvocations);
-				FComputeDataProviderRenderProxy::FDispatchData DispatchData{ NumSubInvocations, 0, 0, ParameterArray.GetStride(), reinterpret_cast<uint8*>(&ParameterArray[0]) };
+				FComputeDataProviderRenderProxy::FDispatchData DispatchData{ NumSubInvocations, bIsUnifiedDispatch, 0, 0, ParameterArray.GetStride(), reinterpret_cast<uint8*>(&ParameterArray[0]) };
 
 				// Iterate shader parameter members to fill the dispatch data structures.
 				// We assume that the members were filled out with a single data interface per member, and that the
@@ -147,7 +179,7 @@ void FComputeGraphTaskWorker::SubmitWork(FRDGBuilder& GraphBuilder, FName InExec
 				// Dispatch work to the render graph.
 				for (int32 SubInvocationIndex = 0; SubInvocationIndex < NumSubInvocations; ++SubInvocationIndex)
 				{
-					TShaderRef<FComputeKernelShader> Shader = Shaders[KernelIndex * NumSubInvocations + SubInvocationIndex];
+					TShaderRef<FComputeKernelShader> Shader = Shaders[ShaderStartPerKernel[KernelIndex] + SubInvocationIndex];
 					const FIntVector GroupCount = FComputeShaderUtils::GetGroupCount(ThreadCounts[SubInvocationIndex], KernelInvocation.KernelGroupSize);
 
 					FComputeShaderUtils::AddPass(
diff --git a/Engine/Plugins/Runtime/ComputeFramework/Source/ComputeFramework/Public/ComputeFramework/ComputeDataInterface.h b/Engine/Plugins/Runtime/ComputeFramework/Source/ComputeFramework/Public/ComputeFramework/ComputeDataInterface.h
--- a/Engine/Plugins/Runtime/ComputeFramework/Source/ComputeFramework/Public/ComputeFramework/ComputeDataInterface.h
+++ b/Engine/Plugins/Runtime/ComputeFramework/Source/ComputeFramework/Public/ComputeFramework/ComputeDataInterface.h
@@ -27,6 +27,8 @@ public:
 	virtual TCHAR const* GetClassName() const PURE_VIRTUAL(UComputeDataInterface::GetClassName, return nullptr;)
 	/** Does the associated UComputeDataProvider provide invocations and thread counts. One and only one data interface per kernel should return true. */
 	virtual bool IsExecutionInterface() const { return false; }
+	/** Return true if the associated UComputeDataProvider holds data that can be combined into a single dispatch invocation. */
+	virtual bool CanSupportUnifiedDispatch() const { return false; }
 	/** Gather compile definitions from the data interface. Any connected kernel will compile with these. */
 	virtual void GetDefines(FComputeKernelDefinitionSet& OutDefinitionSet) const {}
 	/** Gather permutations from the data interface. Any connected kernel will include these in its total compiled permutations. */
diff --git a/Engine/Plugins/Runtime/ComputeFramework/Source/ComputeFramework/Public/ComputeFramework/ComputeDataProvider.h b/Engine/Plugins/Runtime/ComputeFramework/Source/ComputeFramework/Public/ComputeFramework/ComputeDataProvider.h
--- a/Engine/Plugins/Runtime/ComputeFramework/Source/ComputeFramework/Public/ComputeFramework/ComputeDataProvider.h
+++ b/Engine/Plugins/Runtime/ComputeFramework/Source/ComputeFramework/Public/ComputeFramework/ComputeDataProvider.h
@@ -81,6 +81,7 @@ public:
 	struct FDispatchData
 	{
 		int32 NumInvocations;
+		bool bUnifiedDispatch;
 		int32 ParameterStructSize;
 		int32 ParameterBufferOffset;
 		int32 ParameterBufferStride;
