Fix bug where dispatch sizes were being added in all dimensions, causing way too many threads to launch.
Fix is fine for all current execution data interfaces which launch 1d thread groups. Will need refining in the future for 2d or 3d thread groups.

diff --git a/Engine/Plugins/Runtime/ComputeFramework/Source/ComputeFramework/Private/ComputeGraphWorker.cpp b/Engine/Plugins/Runtime/ComputeFramework/Source/ComputeFramework/Private/ComputeGraphWorker.cpp
--- a/Engine/Plugins/Runtime/ComputeFramework/Source/ComputeFramework/Private/ComputeGraphWorker.cpp
+++ b/Engine/Plugins/Runtime/ComputeFramework/Source/ComputeFramework/Private/ComputeGraphWorker.cpp
@@ -147,7 +147,7 @@ void FComputeGraphTaskWorker::SubmitWork(FRDGBuilder& GraphBuilder, FName InExec
 				{
 					for (int32 SubInvocationIndex = 1; SubInvocationIndex < NumSubInvocations; ++SubInvocationIndex)
 					{
-						ThreadCounts[0] += ThreadCounts[SubInvocationIndex];
+						ThreadCounts[0].X += ThreadCounts[SubInvocationIndex].X;
 					}
 					ThreadCounts.SetNum(1);
 					NumSubInvocations = 1;
