CPU performance improvements for using half edge buffers in deformer graph data interface.
The half edge buffer was being uploaded per frame. Now it is stored in a resource owned by the data provider and uploaded once.
This is still less than ideal. We want the resource to be owned by the skel mesh, so that it is cooked instead of created at runtime. But that is work for another day.
Also test if we have any compute worker jobs before kicking off RDG graph. This allows for a fairer CPU comparison between skin cache and deformer graph. (We don't want to add RDG overhead to skin cache only work).
Also added support for reordering the compute work so for optimal GPU execution. Ordering by kernel index instead of graph index allows greater overlap of work on GPU when multiple compute graphs are running.
Added a per graph sort priority so that work sorting doesn't cause any setup graphs to run later than execution graphs.

diff --git a/Engine/Plugins/Animation/DeformerGraph/Source/OptimusCore/Private/DataInterfaces/OptimusDataInterfaceHalfEdge.cpp b/Engine/Plugins/Animation/DeformerGraph/Source/OptimusCore/Private/DataInterfaces/OptimusDataInterfaceHalfEdge.cpp
--- a/Engine/Plugins/Animation/DeformerGraph/Source/OptimusCore/Private/DataInterfaces/OptimusDataInterfaceHalfEdge.cpp
+++ b/Engine/Plugins/Animation/DeformerGraph/Source/OptimusCore/Private/DataInterfaces/OptimusDataInterfaceHalfEdge.cpp
@@ -225,6 +225,67 @@ namespace
 	}
 }
 
+/*
+ * Render resource containing the half edge buffers. 
+ * todo[CF]: We need to move this to the skeletal mesh cooked mesh data to avoid the perf and memory hit everytime we instantiate a deformer.
+ */ 
+class FHalfEdgeBufferResources : public FRenderResource
+{
+public:
+	FHalfEdgeBufferResources(TArray< TArray<int32> >&& InVertexToEdgeData, TArray< TArray<int32> >&& InEdgeToTwinEdgeData)
+	{
+		VertexToEdgeData = MoveTemp(InVertexToEdgeData);
+		EdgeToTwinEdgeData = MoveTemp(InEdgeToTwinEdgeData);
+	}
+
+	void InitRHI() override
+	{
+		FRDGBuilder GraphBuilder(FRHICommandListExecutor::GetImmediateCommandList());
+
+		const int32 NumLods = VertexToEdgeData.Num();
+		VertexToEdgeBuffers.SetNum(NumLods);
+		EdgeToTwinEdgeBuffers.SetNum(NumLods);
+		
+		for (int32 LodIndex = 0; LodIndex < NumLods; ++LodIndex)
+		{
+			FRDGBuffer* VertexToEdgeBuffer = GraphBuilder.CreateBuffer(FRDGBufferDesc::CreateStructuredDesc(sizeof(uint32), VertexToEdgeData[LodIndex].Num()), TEXT("Optimus.VertexToEdge"));
+			GraphBuilder.QueueBufferUpload(VertexToEdgeBuffer, VertexToEdgeData[LodIndex].GetData(), VertexToEdgeData[LodIndex].Num() * sizeof(int32));
+			VertexToEdgeBuffers[LodIndex] = GraphBuilder.ConvertToExternalBuffer(VertexToEdgeBuffer);
+			
+			FRDGBuffer* EdgeToTwinEdgeBuffer = GraphBuilder.CreateBuffer(FRDGBufferDesc::CreateStructuredDesc(sizeof(uint32), EdgeToTwinEdgeData[LodIndex].Num()), TEXT("Optimus.EdgeToTwinEdge"));
+			GraphBuilder.QueueBufferUpload(EdgeToTwinEdgeBuffer, EdgeToTwinEdgeData[LodIndex].GetData(), EdgeToTwinEdgeData[LodIndex].Num() * sizeof(int32));
+			EdgeToTwinEdgeBuffers[LodIndex] = GraphBuilder.ConvertToExternalBuffer(EdgeToTwinEdgeBuffer);
+		}
+
+		GraphBuilder.Execute();
+
+		VertexToEdgeData.Reset();
+		EdgeToTwinEdgeData.Reset();
+	}
+
+	void ReleaseRHI() override
+	{
+		VertexToEdgeBuffers.Reset();
+		EdgeToTwinEdgeBuffers.Reset();
+	}
+
+	TRefCountPtr<FRDGPooledBuffer> GetVertexToEdgeBuffer(int32 InLodIndex)
+	{
+		return VertexToEdgeBuffers[InLodIndex];
+	}
+
+	TRefCountPtr<FRDGPooledBuffer> GetEdgeToTwinBuffer(int32 InLodIndex)
+	{
+		return EdgeToTwinEdgeBuffers[InLodIndex];
+	}
+
+private:
+	TArray< TArray<int32> > VertexToEdgeData;
+	TArray< TArray<int32> > EdgeToTwinEdgeData;
+	TArray< TRefCountPtr<FRDGPooledBuffer> > VertexToEdgeBuffers;
+	TArray< TRefCountPtr<FRDGPooledBuffer> > EdgeToTwinEdgeBuffers;
+};
+
 
 FString UOptimusHalfEdgeDataInterface::GetDisplayName() const
 {
@@ -314,19 +375,25 @@ UComputeDataProvider* UOptimusHalfEdgeDataInterface::CreateDataProvider(TObjectP
 
 	if (Provider->SkinnedMesh != nullptr)
 	{
-		// Build data and store with the provider.
-		// todo[CF]: We need to move this to the skeletal mesh and make part of cooked mesh data instead.
 		FSkeletalMeshRenderData const* SkeletalMeshRenderData = Provider->SkinnedMesh->GetSkeletalMeshRenderData();
 		if (SkeletalMeshRenderData != nullptr)
 		{
-			Provider->VertexToEdgePerLod.SetNum(SkeletalMeshRenderData->NumInlinedLODs);
-			Provider->EdgeToTwinEdgePerLod.SetNum(SkeletalMeshRenderData->NumInlinedLODs);
+			// Build half edge data.
+			TArray< TArray<int32> > VertexToEdgePerLod;
+			VertexToEdgePerLod.SetNum(SkeletalMeshRenderData->NumInlinedLODs);
+			TArray< TArray<int32> > EdgeToTwinEdgePerLod;
+			EdgeToTwinEdgePerLod.SetNum(SkeletalMeshRenderData->NumInlinedLODs);
 
 			for (int32 LodIndex = 0; LodIndex < SkeletalMeshRenderData->NumInlinedLODs; ++LodIndex)
 			{
 				FSkeletalMeshLODRenderData const* LodRenderData = &SkeletalMeshRenderData->LODRenderData[LodIndex];
-				BuildHalfEdgeBuffers(*LodRenderData, Provider->VertexToEdgePerLod[LodIndex], Provider->EdgeToTwinEdgePerLod[LodIndex]);
+				BuildHalfEdgeBuffers(*LodRenderData, VertexToEdgePerLod[LodIndex], EdgeToTwinEdgePerLod[LodIndex]);
 			}
+
+			// Upload data to a resource stored on the provider.
+			Provider->HalfEdgeBuffers = MakeShared<FHalfEdgeBufferResources>(MoveTemp(VertexToEdgePerLod), MoveTemp(EdgeToTwinEdgePerLod));
+
+			BeginInitResource(Provider->HalfEdgeBuffers.Get());
 		}
 	}
 
@@ -334,19 +401,32 @@ UComputeDataProvider* UOptimusHalfEdgeDataInterface::CreateDataProvider(TObjectP
 }
 
 
+void UOptimusHalfEdgeDataProvider::BeginDestroy()
+{
+	Super::BeginDestroy();
+	if (HalfEdgeBuffers)
+	{
+		BeginReleaseResource(HalfEdgeBuffers.Get());
+	}
+	DestroyFence.BeginFence();
+}
+
+bool UOptimusHalfEdgeDataProvider::IsReadyForFinishDestroy()
+{
+	return Super::IsReadyForFinishDestroy() && DestroyFence.IsFenceComplete();
+}
+
 FComputeDataProviderRenderProxy* UOptimusHalfEdgeDataProvider::GetRenderProxy()
 {
-	return new FOptimusHalfEdgeDataProviderProxy(SkinnedMesh, VertexToEdgePerLod, EdgeToTwinEdgePerLod);
+	return new FOptimusHalfEdgeDataProviderProxy(SkinnedMesh, HalfEdgeBuffers);
 }
 
 
 FOptimusHalfEdgeDataProviderProxy::FOptimusHalfEdgeDataProviderProxy(
-	USkinnedMeshComponent* SkinnedMeshComponent, 
-	TArray< TArray<int32> >& InVertexToEdgePerLod,
-	TArray< TArray<int32> >& InEdgeToTwinEdgePerLod)
-	: SkeletalMeshObject(SkinnedMeshComponent != nullptr ? SkinnedMeshComponent->MeshObject : nullptr)
-	, VertexToEdgePerLod(InVertexToEdgePerLod)
-	, EdgeToTwinEdgePerLod(InEdgeToTwinEdgePerLod)
+	USkinnedMeshComponent* InSkinnedMeshComponent, 
+	TSharedPtr<FHalfEdgeBufferResources>& InHalfEdgeBuffers)
+	: SkeletalMeshObject(InSkinnedMeshComponent != nullptr ? InSkinnedMeshComponent->MeshObject : nullptr)
+	, HalfEdgeBuffers(InHalfEdgeBuffers)
 {
 }
 
@@ -356,7 +436,7 @@ bool FOptimusHalfEdgeDataProviderProxy::IsValid(FValidationData const& InValidat
 	{
 		return false;
 	}
-	if (SkeletalMeshObject == nullptr || VertexToEdgePerLod.Num() == 0 || EdgeToTwinEdgePerLod.Num() == 0)
+	if (SkeletalMeshObject == nullptr || !HalfEdgeBuffers.IsValid())
 	{
 		return false;
 	}
@@ -372,19 +452,11 @@ void FOptimusHalfEdgeDataProviderProxy::AllocateResources(FRDGBuilder& GraphBuil
 {
 	const int32 LodIndex = SkeletalMeshObject->GetLOD();
 
-	// todo[CF]: Updating buffer every frame is obviously bad, but just getting things working initially.
-	{
-		TArray<int32> const& Data = VertexToEdgePerLod[LodIndex];
-		VertexToEdgeBuffer = GraphBuilder.CreateBuffer(FRDGBufferDesc::CreateStructuredDesc(sizeof(uint32), Data.Num()), TEXT("Optimus.VertexToEdge"));
-		VertexToEdgeBufferSRV = GraphBuilder.CreateSRV(VertexToEdgeBuffer);
-		GraphBuilder.QueueBufferUpload(VertexToEdgeBuffer, Data.GetData(), Data.Num() * sizeof(uint32), ERDGInitialDataFlags::None);
-	}
-	{
-		TArray<int32> const& Data = EdgeToTwinEdgePerLod[LodIndex];
-		EdgeToTwinEdgeBuffer = GraphBuilder.CreateBuffer(FRDGBufferDesc::CreateStructuredDesc(sizeof(uint32), Data.Num()), TEXT("Optimus.EdgeToTwinEdge"));
-		EdgeToTwinEdgeBufferSRV = GraphBuilder.CreateSRV(EdgeToTwinEdgeBuffer);
-		GraphBuilder.QueueBufferUpload(EdgeToTwinEdgeBuffer, Data.GetData(), Data.Num() * sizeof(uint32), ERDGInitialDataFlags::None);
-	}
+	FRDGBuffer* VertexToEdgeBuffer = GraphBuilder.RegisterExternalBuffer(HalfEdgeBuffers->GetVertexToEdgeBuffer(LodIndex));
+	VertexToEdgeBufferSRV = GraphBuilder.CreateSRV(VertexToEdgeBuffer);
+	
+	FRDGBuffer* EdgeToTwinEdgeBuffer = GraphBuilder.RegisterExternalBuffer(HalfEdgeBuffers->GetEdgeToTwinBuffer(LodIndex));
+	EdgeToTwinEdgeBufferSRV = GraphBuilder.CreateSRV(EdgeToTwinEdgeBuffer);
 }
 
 void FOptimusHalfEdgeDataProviderProxy::GatherDispatchData(FDispatchData const& InDispatchData)
diff --git a/Engine/Plugins/Animation/DeformerGraph/Source/OptimusCore/Private/DataInterfaces/OptimusDataInterfaceHalfEdge.h b/Engine/Plugins/Animation/DeformerGraph/Source/OptimusCore/Private/DataInterfaces/OptimusDataInterfaceHalfEdge.h
--- a/Engine/Plugins/Animation/DeformerGraph/Source/OptimusCore/Private/DataInterfaces/OptimusDataInterfaceHalfEdge.h
+++ b/Engine/Plugins/Animation/DeformerGraph/Source/OptimusCore/Private/DataInterfaces/OptimusDataInterfaceHalfEdge.h
@@ -4,9 +4,11 @@
 
 #include "OptimusComputeDataInterface.h"
 #include "ComputeFramework/ComputeDataProvider.h"
+#include "RenderCommandFence.h"
 #include "OptimusDataInterfaceHalfEdge.generated.h"
 
 class FHalfEdgeDataInterfaceParameters;
+class FHalfEdgeBufferResources;
 class FRDGBuffer;
 class FRDGBufferSRV;
 class FSkeletalMeshObject;
@@ -51,21 +53,25 @@ public:
 	UPROPERTY(EditAnywhere, BlueprintReadWrite, Category = Binding)
 	TObjectPtr<USkinnedMeshComponent> SkinnedMesh = nullptr;
 
+	//~ Begin UObject Interface
+	void BeginDestroy() override;
+	bool IsReadyForFinishDestroy() override;
+	//~ End UObject Interface
+
 	//~ Begin UComputeDataProvider Interface
 	FComputeDataProviderRenderProxy* GetRenderProxy() override;
 	//~ End UComputeDataProvider Interface
 
-	TArray< TArray<int32> > VertexToEdgePerLod;
-	TArray< TArray<int32> > EdgeToTwinEdgePerLod;
+	TSharedPtr<FHalfEdgeBufferResources> HalfEdgeBuffers;
+	FRenderCommandFence DestroyFence;
 };
 
 class FOptimusHalfEdgeDataProviderProxy : public FComputeDataProviderRenderProxy
 {
 public:
 	FOptimusHalfEdgeDataProviderProxy(
-		USkinnedMeshComponent* SkinnedMeshComponent, 
-		TArray< TArray<int32> >& InVertexToEdgePerLod,
-		TArray< TArray<int32> >& InEdgeToTwinEdgeBuffer);
+		USkinnedMeshComponent* InSkinnedMeshComponent, 
+		TSharedPtr<FHalfEdgeBufferResources>& InHalfEdgeBuffers);
 
 	//~ Begin FComputeDataProviderRenderProxy Interface
 	bool IsValid(FValidationData const& InValidationData) const override;
@@ -77,12 +83,8 @@ private:
 	using FParameters = FHalfEdgeDataInterfaceParameters;
 
 	FSkeletalMeshObject* SkeletalMeshObject = nullptr;
+	TSharedPtr<FHalfEdgeBufferResources> HalfEdgeBuffers;
 
-	TArray< TArray<int32> > const& VertexToEdgePerLod;
-	FRDGBuffer* VertexToEdgeBuffer = nullptr;
 	FRDGBufferSRV* VertexToEdgeBufferSRV = nullptr;
-
-	TArray< TArray<int32> > const& EdgeToTwinEdgePerLod;
-	FRDGBuffer* EdgeToTwinEdgeBuffer = nullptr;
 	FRDGBufferSRV* EdgeToTwinEdgeBufferSRV = nullptr;
 };
diff --git a/Engine/Plugins/Animation/DeformerGraph/Source/OptimusCore/Private/OptimusDeformerInstance.cpp b/Engine/Plugins/Animation/DeformerGraph/Source/OptimusCore/Private/OptimusDeformerInstance.cpp
--- a/Engine/Plugins/Animation/DeformerGraph/Source/OptimusCore/Private/OptimusDeformerInstance.cpp
+++ b/Engine/Plugins/Animation/DeformerGraph/Source/OptimusCore/Private/OptimusDeformerInstance.cpp
@@ -309,13 +309,18 @@ void UOptimusDeformerInstance::SetupFromDeformer(UOptimusDeformer* InDeformer)
 	ComputeGraphExecInfos.Reset();
 	GraphsToRunOnNextTick.Reset();
 
-	for (const FOptimusComputeGraphInfo& ComputeGraphInfo : InDeformer->ComputeGraphs)
+	for (int32 GraphIndex = 0; GraphIndex < InDeformer->ComputeGraphs.Num(); ++GraphIndex)
 	{
+		FOptimusComputeGraphInfo const& ComputeGraphInfo = InDeformer->ComputeGraphs[GraphIndex];
 		FOptimusDeformerInstanceExecInfo& Info = ComputeGraphExecInfos.AddDefaulted_GetRef();
 		Info.GraphName = ComputeGraphInfo.GraphName;
 		Info.GraphType = ComputeGraphInfo.GraphType;
 		Info.ComputeGraph = ComputeGraphInfo.ComputeGraph;
 
+		// ComputeGraphs are sorted by the order we want to run them in. 
+		// Using the graph index as our sort priority prevents kernels from the different (but related) graphs running simultaineously.
+		Info.ComputeGraphInstance.SetGraphSortPriority(GraphIndex);
+
 		if (BoundComponents.Num())
 		{
 			UMeshComponent* ActorComponent = MeshComponent.Get();
diff --git a/Engine/Plugins/Runtime/ComputeFramework/Source/ComputeFramework/Private/ComputeFramework/ComputeGraphWorker.h b/Engine/Plugins/Runtime/ComputeFramework/Source/ComputeFramework/Private/ComputeFramework/ComputeGraphWorker.h
--- a/Engine/Plugins/Runtime/ComputeFramework/Source/ComputeFramework/Private/ComputeFramework/ComputeGraphWorker.h
+++ b/Engine/Plugins/Runtime/ComputeFramework/Source/ComputeFramework/Private/ComputeFramework/ComputeGraphWorker.h
@@ -11,6 +11,7 @@
 class FComputeDataProviderRenderProxy;
 class FComputeGraphRenderProxy;
 class FComputeKernelResource;
+class FComputeKernelShader;
 
 /** 
  * Class that manages the scheduling of Compute Graph work.
@@ -23,10 +24,14 @@ public:
 	void Enqueue(
 		FName InExecutionGroupName,
 		FName InOwnerName,
+		uint8 InGraphSortPriority,
 		FComputeGraphRenderProxy const* InGraphRenderProxy, 
 		TArray<FComputeDataProviderRenderProxy*> InDataProviderRenderProxies,
 		FSimpleDelegate InFallbackDelegate);
 
+	/** Has enqueued compute graph work. */
+	bool HasWork(FName InExecutionGroupName) const override;
+
 	/** Submit enqueued compute graph work. */
 	void SubmitWork(
 		FRDGBuilder& GraphBuilder,
@@ -39,6 +44,8 @@ private:
 	{
 		/** Name of owner object that invoked the graph. */
 		FName OwnerName;
+		/** Priority used when sorting work. */
+		uint8 GraphSortPriority = 0;
 		/** Graph render proxy. */
 		FComputeGraphRenderProxy const* GraphRenderProxy = nullptr;
 		/** Data provider render proxies. */
@@ -51,4 +58,38 @@ private:
 
 	/** Map of enqueued work per execution group . */
 	TMap<FName, TArray<FGraphInvocation> > GraphInvocationsPerGroup;
+
+	/** Description of a single dispatch group to be submitted. */
+	struct FSubmitDescription
+	{
+		/**
+		 * Sort key allows us to sort dispatches for optimum scheduling.
+		 * Syncing is usually required between consecutive kernels in a graph.
+		 * So we schedule the first kernels of all the graphs, before all of the second kernels.
+		 * That reduces time sync time, at the expense of memory pressure for buffers that need to stay alive.
+		 * In futuer we may want to add a limit to the number of graphs in flight to avoid memory pressure.
+		 */
+		union
+		{
+			uint32 PackedSortKey = 0;
+			struct
+			{
+				uint32 GraphIndex : 12;			// Graph index.
+				uint32 KernelIndex : 12;		// Kernel index within the graph.
+				uint32 GraphSortPriority : 8;	// Externally defined sort priority to maintain inter-graph dependencies.
+			};
+		};
+
+		/** Track the index into our collected shader array. */
+		uint32 ShaderIndex : 15;
+		/** Track if this is Unified Dispatch. */
+		uint32 bIsUnified : 1;
+	};
+
+	// These arrays could be local to FComputeGraphTaskWorker::SubmitWork() but we 
+	// store them with class and Reset() them at each usage to avoid per frame array allocations.
+	TArray<FSubmitDescription> SubmitDescs;
+	TArray<TShaderRef<FComputeKernelShader>> Shaders;
+	TArray<int32> PermutationIds;
+	TArray<FIntVector> ThreadCounts;
 };
diff --git a/Engine/Plugins/Runtime/ComputeFramework/Source/ComputeFramework/Private/ComputeGraphInstance.cpp b/Engine/Plugins/Runtime/ComputeFramework/Source/ComputeFramework/Private/ComputeGraphInstance.cpp
--- a/Engine/Plugins/Runtime/ComputeFramework/Source/ComputeFramework/Private/ComputeGraphInstance.cpp
+++ b/Engine/Plugins/Runtime/ComputeFramework/Source/ComputeFramework/Private/ComputeGraphInstance.cpp
@@ -71,10 +71,10 @@ bool FComputeGraphInstance::EnqueueWork(UComputeGraph* InComputeGraph, FSceneInt
 	}
 
 	ENQUEUE_RENDER_COMMAND(ComputeFrameworkEnqueueExecutionCommand)(
-		[ComputeGraphWorker, InExecutionGroupName, InOwnerName, GraphRenderProxy, MovedDataProviderRenderProxies = MoveTemp(DataProviderRenderProxies), InFallbackDelegate](FRHICommandListImmediate& RHICmdList)
+		[ComputeGraphWorker, InExecutionGroupName, InOwnerName, SortPriority = GraphSortPriority, GraphRenderProxy, MovedDataProviderRenderProxies = MoveTemp(DataProviderRenderProxies), InFallbackDelegate](FRHICommandListImmediate& RHICmdList)
 		{
 			// Compute graph scheduler will take ownership of the provider proxies.
-			ComputeGraphWorker->Enqueue(InExecutionGroupName, InOwnerName, GraphRenderProxy, MovedDataProviderRenderProxies, InFallbackDelegate);
+			ComputeGraphWorker->Enqueue(InExecutionGroupName, InOwnerName, SortPriority, GraphRenderProxy, MovedDataProviderRenderProxies, InFallbackDelegate);
 		});
 
 	return true;
diff --git a/Engine/Plugins/Runtime/ComputeFramework/Source/ComputeFramework/Private/ComputeGraphWorker.cpp b/Engine/Plugins/Runtime/ComputeFramework/Source/ComputeFramework/Private/ComputeGraphWorker.cpp
--- a/Engine/Plugins/Runtime/ComputeFramework/Source/ComputeFramework/Private/ComputeGraphWorker.cpp
+++ b/Engine/Plugins/Runtime/ComputeFramework/Source/ComputeFramework/Private/ComputeGraphWorker.cpp
@@ -2,6 +2,7 @@
 
 #include "ComputeFramework/ComputeGraphWorker.h"
 
+#include "Algo/Sort.h"
 #include "ComputeFramework/ComputeKernel.h"
 #include "ComputeFramework/ComputeKernelPermutationVector.h"
 #include "ComputeFramework/ComputeKernelShader.h"
@@ -14,196 +15,234 @@
 
 DECLARE_GPU_STAT_NAMED(ComputeFramework_ExecuteBatches, TEXT("ComputeFramework::ExecuteBatches"));
 
-void FComputeGraphTaskWorker::Enqueue(FName InExecutionGroupName, FName InOwnerName, FComputeGraphRenderProxy const* InGraphRenderProxy, TArray<FComputeDataProviderRenderProxy*> InDataProviderRenderProxies, FSimpleDelegate InFallbackDelegate)
+static TAutoConsoleVariable<int32> CVarComputeFrameworkSortSubmit(
+	TEXT("r.ComputeFramework.SortSubmit"),
+	1,
+	TEXT("Sort submission of work to GPU for optimal scheduling."),
+	ECVF_RenderThreadSafe
+);
+
+
+void FComputeGraphTaskWorker::Enqueue(
+	FName InExecutionGroupName, 
+	FName InOwnerName, 
+	uint8 InGraphSortPriority,
+	FComputeGraphRenderProxy const* InGraphRenderProxy, 
+	TArray<FComputeDataProviderRenderProxy*> InDataProviderRenderProxies, 
+	FSimpleDelegate InFallbackDelegate)
 {
 	FGraphInvocation& GraphInvocation = GraphInvocationsPerGroup.FindOrAdd(InExecutionGroupName).AddDefaulted_GetRef();
 	GraphInvocation.OwnerName = InOwnerName;
+	GraphInvocation.GraphSortPriority = InGraphSortPriority;
 	GraphInvocation.GraphRenderProxy = InGraphRenderProxy;
 	GraphInvocation.DataProviderRenderProxies = MoveTemp(InDataProviderRenderProxies);
 	GraphInvocation.FallbackDelegate = InFallbackDelegate;
 }
 
+bool FComputeGraphTaskWorker::HasWork(FName InExecutionGroupName) const
+{
+	TArray<FGraphInvocation> const* GraphInvocations = GraphInvocationsPerGroup.Find(InExecutionGroupName);
+	return GraphInvocations != nullptr && GraphInvocations->Num();
+}
+
 void FComputeGraphTaskWorker::SubmitWork(FRDGBuilder& GraphBuilder, FName InExecutionGroupName, ERHIFeatureLevel::Type FeatureLevel)
 {
-	TArray<FGraphInvocation>& GraphInvocations = GraphInvocationsPerGroup.FindOrAdd(InExecutionGroupName);
-	if (GraphInvocations.IsEmpty())
-	{
-		return;
-	}
+	TRACE_CPUPROFILER_EVENT_SCOPE("ComputeFramework::ExecuteBatches");
+	RDG_EVENT_SCOPE(GraphBuilder, "ComputeFramework::ExecuteBatches");
+	RDG_GPU_STAT_SCOPE(GraphBuilder, ComputeFramework_ExecuteBatches);
 
+	// Reset our scratch memory arrays.
+	SubmitDescs.Reset();
+	Shaders.Reset();
+
+	TArray<FGraphInvocation> const& GraphInvocations = GraphInvocationsPerGroup.FindChecked(InExecutionGroupName);
+	for (int32 GraphIndex = 0; GraphIndex < GraphInvocations.Num(); ++GraphIndex)
 	{
-		SCOPED_DRAW_EVENTF(GraphBuilder.RHICmdList, ComputeFramework_ExecuteBatches, TEXT("ComputeFramework::ExecuteBatches"));
-		SCOPED_GPU_STAT(GraphBuilder.RHICmdList, ComputeFramework_ExecuteBatches);
+		FGraphInvocation const& GraphInvocation = GraphInvocations[GraphIndex];
+		FComputeGraphRenderProxy const* GraphRenderProxy = GraphInvocation.GraphRenderProxy;
+		const int32 NumKernels = GraphRenderProxy->KernelInvocations.Num();
+
+		const int32 BaseSubmitDescIndex = SubmitDescs.Num();
+		SubmitDescs.Reserve(BaseSubmitDescIndex + NumKernels);
+		const int32 BaseShaderIndex = Shaders.Num();
 
-		for (int32 GraphIndex = 0; GraphIndex < GraphInvocations.Num(); ++GraphIndex)
+		// Gather shaders and validate the DataInterfaces.
+		// If validation fails or shaders are awaiting compilation we will not run the graph.
+		bool bIsValid = true;
+		for (int32 KernelIndex = 0; bIsValid && KernelIndex < NumKernels; ++KernelIndex)
 		{
-			FGraphInvocation const& GraphInvocation = GraphInvocations[GraphIndex];
-			FComputeGraphRenderProxy const* GraphRenderProxy = GraphInvocation.GraphRenderProxy;
-
-			TArray<TShaderRef<FComputeKernelShader>> Shaders;
-			TArray<int32> ShaderStartPerKernel;
-			ShaderStartPerKernel.Reserve(GraphRenderProxy->KernelInvocations.Num());
-			TArray<bool> UnifiedDispatchPerKernel;
-			UnifiedDispatchPerKernel.SetNumZeroed(GraphRenderProxy->KernelInvocations.Num());
-
-			// Validation phase.
-			// Check if all DataInterfaces are valid.
-			// At the same time gather the permutation id so that we can validate if shader is compiled.
-			bool bIsValid = true;
-			for (int32 KernelIndex = 0; bIsValid && KernelIndex < GraphRenderProxy->KernelInvocations.Num(); ++KernelIndex)
-			{
-				FComputeGraphRenderProxy::FKernelInvocation const& KernelInvocation = GraphRenderProxy->KernelInvocations[KernelIndex];
+			FSubmitDescription& SubmitDesc = SubmitDescs.AddZeroed_GetRef();
+			SubmitDesc.GraphIndex = GraphIndex;
+			SubmitDesc.KernelIndex = KernelIndex;
+			SubmitDesc.GraphSortPriority = GraphInvocation.GraphSortPriority;
+			SubmitDesc.ShaderIndex = Shaders.Num();
+
+			FComputeGraphRenderProxy::FKernelInvocation const& KernelInvocation = GraphRenderProxy->KernelInvocations[KernelIndex];
 
-				TArray<FIntVector> ThreadCounts;
-				const int32 NumSubInvocations = GraphInvocation.DataProviderRenderProxies[KernelInvocation.ExecutionProviderIndex]->GetDispatchThreadCount(ThreadCounts);
+			// Reset our scratch memory arrays.
+			PermutationIds.Reset();
+			ThreadCounts.Reset();
 
-				// Iterate shader parameter members to fill the dispatch data structures.
-				// We assume that the members were filled out with a single data interface per member, and that the
-				// order is the same one defined in the KernelInvocation.BoundProviderIndices.
-				TArray<FShaderParametersMetadata::FMember> const& ParamMembers = KernelInvocation.ShaderParameterMetadata->GetMembers();
+			const int32 NumSubInvocations = GraphInvocation.DataProviderRenderProxies[KernelInvocation.ExecutionProviderIndex]->GetDispatchThreadCount(ThreadCounts);
 
-				FComputeDataProviderRenderProxy::FPermutationData PermutationData{ NumSubInvocations, GraphRenderProxy->ShaderPermutationVectors[KernelIndex] };
-				PermutationData.PermutationIds.AddZeroed(NumSubInvocations);
+			// Iterate shader parameter members to fill the dispatch data structures.
+			// We assume that the members were filled out with a single data interface per member, and that the
+			// order is the same one defined in the KernelInvocation.BoundProviderIndices.
+			TArray<FShaderParametersMetadata::FMember> const& ParamMembers = KernelInvocation.ShaderParameterMetadata->GetMembers();
 
-				for (int32 MemberIndex = 0; bIsValid && MemberIndex < ParamMembers.Num(); ++MemberIndex)
+			FComputeDataProviderRenderProxy::FPermutationData PermutationData{ NumSubInvocations, GraphRenderProxy->ShaderPermutationVectors[KernelIndex], MoveTemp(PermutationIds) };
+			PermutationData.PermutationIds.SetNumZeroed(NumSubInvocations);
+
+			for (int32 MemberIndex = 0; bIsValid && MemberIndex < ParamMembers.Num(); ++MemberIndex)
+			{
+				FShaderParametersMetadata::FMember const& Member = ParamMembers[MemberIndex];
+				if (ensure(Member.GetBaseType() == EUniformBufferBaseType::UBMT_NESTED_STRUCT))
 				{
-					FShaderParametersMetadata::FMember const& Member = ParamMembers[MemberIndex];
-					if (ensure(Member.GetBaseType() == EUniformBufferBaseType::UBMT_NESTED_STRUCT))
+					const int32 DataProviderIndex = KernelInvocation.BoundProviderIndices[MemberIndex];
+					FComputeDataProviderRenderProxy* DataProvider = GraphInvocation.DataProviderRenderProxies[DataProviderIndex];
+					if (ensure(DataProvider != nullptr))
 					{
-						const int32 DataProviderIndex = KernelInvocation.BoundProviderIndices[MemberIndex];
-						FComputeDataProviderRenderProxy* DataProvider = GraphInvocation.DataProviderRenderProxies[DataProviderIndex];
-						if (ensure(DataProvider != nullptr))
-						{
-							FComputeDataProviderRenderProxy::FValidationData ValidationData{ NumSubInvocations, (int32)Member.GetStructMetadata()->GetSize() };
-							bIsValid &= DataProvider->IsValid(ValidationData);
+						FComputeDataProviderRenderProxy::FValidationData ValidationData{ NumSubInvocations, (int32)Member.GetStructMetadata()->GetSize() };
+						bIsValid &= DataProvider->IsValid(ValidationData);
 
-							if (bIsValid)
-							{
-								DataProvider->GatherPermutations(PermutationData);
-							}
+						if (bIsValid)
+						{
+							DataProvider->GatherPermutations(PermutationData);
 						}
 					}
 				}
+			}
 
-				// Get shader. This can fail if compilation is pending.
-				ShaderStartPerKernel.Add(Shaders.Num());
-				Shaders.Reserve(Shaders.Num() + NumSubInvocations);
-				for (int32 SubInvocationIndex = 0; bIsValid && SubInvocationIndex < NumSubInvocations; ++SubInvocationIndex)
+			// Get shader. This can fail if compilation is pending.
+			for (int32 SubInvocationIndex = 0; bIsValid && SubInvocationIndex < NumSubInvocations; ++SubInvocationIndex)
+			{
+				TShaderRef<FComputeKernelShader> Shader = KernelInvocation.KernelResource->GetShader(PermutationData.PermutationIds[SubInvocationIndex]);
+				bIsValid &= Shader.IsValid();
+				Shaders.Add(Shader);
+			}
+
+			// Check if we can do unified dispatch and apply that if we can.
+			if (bIsValid && KernelInvocation.bSupportsUnifiedDispatch && NumSubInvocations > 1)
+			{
+				bool bSupportsUnifiedDispatch = true;
+				for (int32 SubInvocationIndex = 1; bSupportsUnifiedDispatch && SubInvocationIndex < NumSubInvocations; ++SubInvocationIndex)
 				{
-					TShaderRef<FComputeKernelShader> Shader = KernelInvocation.KernelResource->GetShader(PermutationData.PermutationIds[SubInvocationIndex]);
-					bIsValid &= Shader.IsValid();
-					Shaders.Add(Shader);
+					bSupportsUnifiedDispatch &= Shaders[SubmitDesc.ShaderIndex + SubInvocationIndex] == Shaders[SubmitDesc.ShaderIndex];
 				}
 
-				// Check if we can do unified dispatch and apply that if we can.
-				if (bIsValid && KernelInvocation.bSupportsUnifiedDispatch && NumSubInvocations > 1)
+				if (bSupportsUnifiedDispatch)
 				{
-					bool bSupportsUnifiedDispatch = true;
-					for (int32 SubInvocationIndex = 1; bSupportsUnifiedDispatch && SubInvocationIndex < NumSubInvocations; ++SubInvocationIndex)
-					{
-						bSupportsUnifiedDispatch &= Shaders[ShaderStartPerKernel[KernelIndex] + SubInvocationIndex] == Shaders[ShaderStartPerKernel[KernelIndex]];
-					}
-
-					if (bSupportsUnifiedDispatch)
-					{
-						UnifiedDispatchPerKernel[KernelIndex] = true;
-						Shaders.SetNum(Shaders.Num() - NumSubInvocations + 1);
-					}
+					SubmitDesc.bIsUnified = true;
+					Shaders.SetNum(SubmitDesc.ShaderIndex + 1, /*bAllowShrinking*/false);
 				}
 			}
 
-			// If we can't run the graph for any reason, back out now and apply fallback logic.
-			if (!bIsValid)
-			{
-				GraphInvocation.FallbackDelegate.ExecuteIfBound();
-				continue;
-			}
+			// Move our scratch array back for subsequent reuse.
+			PermutationIds = MoveTemp(PermutationData.PermutationIds);
+		}
 
-			// From here on we are committed to submitting the work to the GPU.
-			RDG_EVENT_SCOPE(GraphBuilder, "%s:%s", *GraphInvocation.OwnerName.ToString(), *GraphRenderProxy->GraphName.ToString());
+		// If we can't run the graph for any reason, back out now and apply fallback logic.
+		if (!bIsValid)
+		{
+			SubmitDescs.SetNum(BaseSubmitDescIndex, /*bAllowShrinking*/false);
+			Shaders.SetNum(BaseShaderIndex, /*bAllowShrinking*/false);
+			GraphInvocation.FallbackDelegate.ExecuteIfBound();
+			continue;
+		}
 
-			// Do resource allocation for all the data providers in the graph.
-			for (int32 DataProviderIndex = 0; DataProviderIndex < GraphInvocation.DataProviderRenderProxies.Num(); ++DataProviderIndex)
+		// Allocate RDG resources for all the data providers in the graph.
+		for (int32 DataProviderIndex = 0; DataProviderIndex < GraphInvocation.DataProviderRenderProxies.Num(); ++DataProviderIndex)
+		{
+			FComputeDataProviderRenderProxy* DataProvider = GraphInvocation.DataProviderRenderProxies[DataProviderIndex];
+			if (DataProvider != nullptr)
 			{
-				FComputeDataProviderRenderProxy* DataProvider = GraphInvocation.DataProviderRenderProxies[DataProviderIndex];
-				if (DataProvider != nullptr)
-				{
-					DataProvider->AllocateResources(GraphBuilder);
-				}
+				DataProvider->AllocateResources(GraphBuilder);
 			}
+		}
+	}
 
-			// Iterate the graph kernels to collect shader bindings and dispatch work.
-			for (int32 KernelIndex = 0; KernelIndex < GraphRenderProxy->KernelInvocations.Num(); ++KernelIndex)
-			{
-				FComputeGraphRenderProxy::FKernelInvocation const& KernelInvocation = GraphRenderProxy->KernelInvocations[KernelIndex];
+	if (CVarComputeFrameworkSortSubmit.GetValueOnRenderThread() != 0)
+	{
+		// Sort for optimal dispatch.
+		Algo::Sort(SubmitDescs, [](const FSubmitDescription& LHS, const FSubmitDescription& RHS) { return LHS.PackedSortKey < RHS.PackedSortKey; });
+	}
 
-				RDG_EVENT_SCOPE(GraphBuilder, "%s", *KernelInvocation.KernelName);
+	for (FSubmitDescription const& SubmitDesc : SubmitDescs)
+	{
+		const int32 GraphIndex = SubmitDesc.GraphIndex;
+		FGraphInvocation const& GraphInvocation = GraphInvocations[GraphIndex];
+		FComputeGraphRenderProxy const* GraphRenderProxy = GraphInvocation.GraphRenderProxy;
 
-				TArray<FIntVector> ThreadCounts;
-				int32 NumSubInvocations = GraphInvocation.DataProviderRenderProxies[KernelInvocation.ExecutionProviderIndex]->GetDispatchThreadCount(ThreadCounts);
-				
-				bool bIsUnifiedDispatch = UnifiedDispatchPerKernel[KernelIndex];
-				if (bIsUnifiedDispatch)
-				{
-					for (int32 SubInvocationIndex = 1; SubInvocationIndex < NumSubInvocations; ++SubInvocationIndex)
-					{
-						ThreadCounts[0].X += ThreadCounts[SubInvocationIndex].X;
-					}
-					ThreadCounts.SetNum(1);
-					NumSubInvocations = 1;
-				}
+		const int32 KernelIndex = SubmitDesc.KernelIndex;
+		FComputeGraphRenderProxy::FKernelInvocation const& KernelInvocation = GraphRenderProxy->KernelInvocations[KernelIndex];
 
-				TStridedView<FComputeKernelShader::FParameters> ParameterArray = GraphBuilder.AllocParameters<FComputeKernelShader::FParameters>(KernelInvocation.ShaderParameterMetadata, NumSubInvocations);
-				FComputeDataProviderRenderProxy::FDispatchData DispatchData{ NumSubInvocations, bIsUnifiedDispatch, 0, 0, ParameterArray.GetStride(), reinterpret_cast<uint8*>(&ParameterArray[0]) };
+		RDG_EVENT_SCOPE(GraphBuilder, "%s:%s:%s", *GraphInvocation.OwnerName.ToString(), *GraphRenderProxy->GraphName.ToString(), *KernelInvocation.KernelName);
 
-				// Iterate shader parameter members to fill the dispatch data structures.
-				// We assume that the members were filled out with a single data interface per member, and that the
-				// order is the same one defined in the KernelInvocation.BoundProviderIndices.
-				TArray<FShaderParametersMetadata::FMember> const& ParamMembers = KernelInvocation.ShaderParameterMetadata->GetMembers();
-				for (int32 MemberIndex = 0; MemberIndex < ParamMembers.Num(); ++MemberIndex)
-				{
-					FShaderParametersMetadata::FMember const& Member = ParamMembers[MemberIndex];
-					if (ensure(Member.GetBaseType() == EUniformBufferBaseType::UBMT_NESTED_STRUCT))
-					{
-						const int32 DataProviderIndex = KernelInvocation.BoundProviderIndices[MemberIndex];
-						FComputeDataProviderRenderProxy* DataProvider = GraphInvocation.DataProviderRenderProxies[DataProviderIndex];
-						if (ensure(DataProvider != nullptr))
-						{
-							DispatchData.ParameterStructSize = Member.GetStructMetadata()->GetSize();
-							DispatchData.ParameterBufferOffset = Member.GetOffset();
-							DataProvider->GatherDispatchData(DispatchData);
-						}
-					}
-				}
+		//todo[CF]: GetDispatchThreadCount() should take the bIsUnified flag directly.
+		ThreadCounts.Reset();
+		int32 NumSubInvocations = GraphInvocation.DataProviderRenderProxies[KernelInvocation.ExecutionProviderIndex]->GetDispatchThreadCount(ThreadCounts);
+
+		bool bIsUnifiedDispatch = SubmitDesc.bIsUnified;
+		if (bIsUnifiedDispatch)
+		{
+			for (int32 SubInvocationIndex = 1; SubInvocationIndex < NumSubInvocations; ++SubInvocationIndex)
+			{
+				ThreadCounts[0].X += ThreadCounts[SubInvocationIndex].X;
+			}
+			ThreadCounts.SetNum(1);
+			NumSubInvocations = 1;
+		}
 
-				// Dispatch work to the render graph.
-				for (int32 SubInvocationIndex = 0; SubInvocationIndex < NumSubInvocations; ++SubInvocationIndex)
+		// Allocate parameters buffer and fill from data providers.
+		TStridedView<FComputeKernelShader::FParameters> ParameterArray = GraphBuilder.AllocParameters<FComputeKernelShader::FParameters>(KernelInvocation.ShaderParameterMetadata, NumSubInvocations);
+		FComputeDataProviderRenderProxy::FDispatchData DispatchData{ NumSubInvocations, bIsUnifiedDispatch, 0, 0, ParameterArray.GetStride(), reinterpret_cast<uint8*>(&ParameterArray[0]) };
+
+		// Iterate shader parameter members to fill the dispatch data structures.
+		// We assume that the members were filled out with a single data interface per member, and that the
+		// order is the same one defined in the KernelInvocation.BoundProviderIndices.
+		TArray<FShaderParametersMetadata::FMember> const& ParamMembers = KernelInvocation.ShaderParameterMetadata->GetMembers();
+		for (int32 MemberIndex = 0; MemberIndex < ParamMembers.Num(); ++MemberIndex)
+		{
+			FShaderParametersMetadata::FMember const& Member = ParamMembers[MemberIndex];
+			if (ensure(Member.GetBaseType() == EUniformBufferBaseType::UBMT_NESTED_STRUCT))
+			{
+				const int32 DataProviderIndex = KernelInvocation.BoundProviderIndices[MemberIndex];
+				FComputeDataProviderRenderProxy* DataProvider = GraphInvocation.DataProviderRenderProxies[DataProviderIndex];
+				if (ensure(DataProvider != nullptr))
 				{
-					TShaderRef<FComputeKernelShader> Shader = Shaders[ShaderStartPerKernel[KernelIndex] + SubInvocationIndex];
-					const FIntVector GroupCount = FComputeShaderUtils::GetGroupCount(ThreadCounts[SubInvocationIndex], KernelInvocation.KernelGroupSize);
-
-					FComputeShaderUtils::AddPass(
-						GraphBuilder,
-						{},
-						ERDGPassFlags::Compute | ERDGPassFlags::NeverCull,
-						Shader,
-						KernelInvocation.ShaderParameterMetadata,
-						&ParameterArray[SubInvocationIndex],
-						GroupCount
-					);
+					DispatchData.ParameterStructSize = Member.GetStructMetadata()->GetSize();
+					DispatchData.ParameterBufferOffset = Member.GetOffset();
+					DataProvider->GatherDispatchData(DispatchData);
 				}
 			}
 		}
 
-		// Release any graph resources at the end of graph execution.
-		GraphBuilder.AddPass(
-			RDG_EVENT_NAME("Release Data Providers"),
-			ERDGPassFlags::None,
-			[this, InExecutionGroupName](FRHICommandList&)
-			{
-				GraphInvocationsPerGroup.FindChecked(InExecutionGroupName).Reset();
-			});
+		// Dispatch work to the render graph.
+		for (int32 SubInvocationIndex = 0; SubInvocationIndex < NumSubInvocations; ++SubInvocationIndex)
+		{
+			TShaderRef<FComputeKernelShader> Shader = Shaders[SubmitDesc.ShaderIndex + SubInvocationIndex];
+			const FIntVector GroupCount = FComputeShaderUtils::GetGroupCount(ThreadCounts[SubInvocationIndex], KernelInvocation.KernelGroupSize);
+
+			FComputeShaderUtils::AddPass(
+				GraphBuilder,
+				{},
+				ERDGPassFlags::Compute | ERDGPassFlags::NeverCull,
+				Shader,
+				KernelInvocation.ShaderParameterMetadata,
+				&ParameterArray[SubInvocationIndex],
+				GroupCount
+			);
+		}
 	}
+
+	// Release any graph resources at the end of graph execution.
+	GraphBuilder.AddPass(
+		{},
+		ERDGPassFlags::None,
+		[this, InExecutionGroupName](FRHICommandList&)
+		{
+			GraphInvocationsPerGroup.FindChecked(InExecutionGroupName).Reset();
+		});
 }
 
 FComputeGraphTaskWorker::FGraphInvocation::~FGraphInvocation()
diff --git a/Engine/Plugins/Runtime/ComputeFramework/Source/ComputeFramework/Public/ComputeFramework/ComputeGraphInstance.h b/Engine/Plugins/Runtime/ComputeFramework/Source/ComputeFramework/Public/ComputeFramework/ComputeGraphInstance.h
--- a/Engine/Plugins/Runtime/ComputeFramework/Source/ComputeFramework/Public/ComputeFramework/ComputeGraphInstance.h
+++ b/Engine/Plugins/Runtime/ComputeFramework/Source/ComputeFramework/Public/ComputeFramework/ComputeGraphInstance.h
@@ -19,6 +19,12 @@ struct COMPUTEFRAMEWORK_API FComputeGraphInstance
 	GENERATED_USTRUCT_BODY();
 
 public:
+	/** 
+	 * Set the priority used when sorting work. 
+	 * Kernels in instances with a lower sort prioirty will always be submitted first.
+	 */
+	void SetGraphSortPriority(uint8 InPriority) { GraphSortPriority = InPriority; }
+
 	/** 
 	 * Create the Data Provider objects for a single binding of a ComputeGraph. 
 	 * The type of binding object is expected to match the associated Binding on the UComputeGraph.
@@ -38,4 +44,7 @@ private:
 	/** The currently bound Data Provider objects. */
 	UPROPERTY(Transient)
 	TArray< TObjectPtr<UComputeDataProvider> > DataProviders;
+
+	/** Priority used when sorting work. */
+	uint8 GraphSortPriority = 0;
 };
diff --git a/Engine/Source/Runtime/Engine/Private/LevelTick.cpp b/Engine/Source/Runtime/Engine/Private/LevelTick.cpp
--- a/Engine/Source/Runtime/Engine/Private/LevelTick.cpp
+++ b/Engine/Source/Runtime/Engine/Private/LevelTick.cpp
@@ -1009,9 +1009,12 @@ void EndSendEndOfFrameUpdatesDrawEvent(FSendAllEndOfFrameUpdates& SendAllEndOfFr
 				SCOPED_GPU_STAT(RHICmdList, ComputeTaskWorkerUpdates);
 				for (IComputeTaskWorker* ComputeTaskWorker : ComputeTaskWorkers)
 				{
-					FRDGBuilder GraphBuilder(RHICmdList);
-					ComputeTaskWorker->SubmitWork(GraphBuilder, ComputeTaskExecutionGroup::EndOfFrameUpdate, FeatureLevel);
-					GraphBuilder.Execute();
+					if (ComputeTaskWorker->HasWork(ComputeTaskExecutionGroup::EndOfFrameUpdate))
+					{
+						FRDGBuilder GraphBuilder(RHICmdList, RDG_EVENT_NAME("ComputeTaskWorker"));
+						ComputeTaskWorker->SubmitWork(GraphBuilder, ComputeTaskExecutionGroup::EndOfFrameUpdate, FeatureLevel);
+						GraphBuilder.Execute();
+					}
 				}
 			}
 		});
diff --git a/Engine/Source/Runtime/Engine/Public/ComputeWorkerInterface.h b/Engine/Source/Runtime/Engine/Public/ComputeWorkerInterface.h
--- a/Engine/Source/Runtime/Engine/Public/ComputeWorkerInterface.h
+++ b/Engine/Source/Runtime/Engine/Public/ComputeWorkerInterface.h
@@ -19,6 +19,9 @@ class ENGINE_API IComputeTaskWorker
 public:
 	virtual ~IComputeTaskWorker() {}
 
+	/** Returns true if there is any scheduled work. */
+	virtual bool HasWork(FName InExecutionGroupName) const = 0;
+
 	/** Add any scheduled work to an RDGBuilder ready for execution. */
 	virtual void SubmitWork(FRDGBuilder& GraphBuilder, FName InExecutionGroupName, ERHIFeatureLevel::Type FeatureLevel) = 0;
 };
