ComputeFramework: corrected a typo when determining dispatch type for data providers bound to a kernel.

diff --git a/Engine/Plugins/Runtime/ComputeFramework/Source/ComputeFramework/Private/ComputeGraphWorker.cpp b/Engine/Plugins/Runtime/ComputeFramework/Source/ComputeFramework/Private/ComputeGraphWorker.cpp
--- a/Engine/Plugins/Runtime/ComputeFramework/Source/ComputeFramework/Private/ComputeGraphWorker.cpp
+++ b/Engine/Plugins/Runtime/ComputeFramework/Source/ComputeFramework/Private/ComputeGraphWorker.cpp
@@ -219,7 +219,7 @@ void FComputeGraphTaskWorker::SubmitWork(FRDGBuilder& GraphBuilder, FName InExec
 					// so user may want to consider using a dummy buffer that maps section index to the indices of secondary buffers
 					// for example, given a non-unified kernel, primary and secondary components sharing the same vertex count, we might want to create a buffer
 					// in the primary group that is simply [0,1,2...,NumVerts-1], which we can then index into to map section vert index to the global vert index
-					DispatchData.bUnifiedDispatch = KernelInvocation.BoundProviderIsPrimary[MemberIndex]? DispatchData.bUnifiedDispatch : true;
+					DispatchData.bUnifiedDispatch = KernelInvocation.BoundProviderIsPrimary[MemberIndex]? bIsUnifiedDispatch : true;
 					DispatchData.ParameterStructSize = Member.GetStructMetadata()->GetSize();
 					DispatchData.ParameterBufferOffset = Member.GetOffset();
 					DataProvider->GatherDispatchData(DispatchData);
