Added support for aborting compute graph work when releasing resources.

diff --git a/Engine/Plugins/Animation/DeformerGraph/Source/OptimusCore/Private/OptimusDeformerInstance.cpp b/Engine/Plugins/Animation/DeformerGraph/Source/OptimusCore/Private/OptimusDeformerInstance.cpp
--- a/Engine/Plugins/Animation/DeformerGraph/Source/OptimusCore/Private/OptimusDeformerInstance.cpp
+++ b/Engine/Plugins/Animation/DeformerGraph/Source/OptimusCore/Private/OptimusDeformerInstance.cpp
@@ -276,7 +276,9 @@ UOptimusComponentSourceBinding const* UOptimusDeformerInstanceSettings::GetCompo
 
 void UOptimusDeformerInstance::SetMeshComponent(UMeshComponent* InMeshComponent)
 { 
+	check(InMeshComponent);
 	MeshComponent = InMeshComponent;
+	Scene = MeshComponent->GetScene();
 }
 
 void UOptimusDeformerInstance::SetInstanceSettings(UOptimusDeformerInstanceSettings* InInstanceSettings)
@@ -516,13 +518,20 @@ void UOptimusDeformerInstance::AllocateResources()
 
 void UOptimusDeformerInstance::ReleaseResources()
 {
-	if (BufferPool)
+	if (Scene || BufferPool)
 	{
-		ENQUEUE_RENDER_COMMAND(FOptimusReleasePoolMemory)(
-			[BufferPool=MoveTemp(BufferPool)](FRHICommandListImmediate& InCmdList)
+		ENQUEUE_RENDER_COMMAND(OptimusReleaseResources)([BufferPool=MoveTemp(BufferPool), Scene = Scene, OwnerPointer = this] (FRHICommandListImmediate& InCmdList)
+		{
+			if (Scene)
+			{
+				ComputeFramework::AbortWork(Scene, OwnerPointer);
+			}
+
+			if (BufferPool)
 			{
 				BufferPool->ReleaseResources();
-			});
+			}
+		});
 	}
 }
 
@@ -559,7 +568,7 @@ void UOptimusDeformerInstance::EnqueueWork(FEnqueueWorkDesc const& InDesc)
 		{
 			if (Info.GraphType == EOptimusNodeGraphType::Update || GraphsToRun.Contains(Info.GraphName))
 			{
-				bIsWorkEnqueued |= Info.ComputeGraphInstance.EnqueueWork(Info.ComputeGraph, InDesc.Scene, ExecutionGroupName, InDesc.OwnerName, InDesc.FallbackDelegate);
+				bIsWorkEnqueued |= Info.ComputeGraphInstance.EnqueueWork(Info.ComputeGraph, InDesc.Scene, ExecutionGroupName, InDesc.OwnerName, InDesc.FallbackDelegate, this);
 			}
 		}
 	}
diff --git a/Engine/Plugins/Animation/DeformerGraph/Source/OptimusCore/Public/OptimusDeformerInstance.h b/Engine/Plugins/Animation/DeformerGraph/Source/OptimusCore/Public/OptimusDeformerInstance.h
--- a/Engine/Plugins/Animation/DeformerGraph/Source/OptimusCore/Public/OptimusDeformerInstance.h
+++ b/Engine/Plugins/Animation/DeformerGraph/Source/OptimusCore/Public/OptimusDeformerInstance.h
@@ -269,6 +269,7 @@ private:
 	FCriticalSection GraphsToRunOnNextTickLock;
 
 	FOptimusPersistentBufferPoolPtr BufferPool;
+	FSceneInterface* Scene = nullptr;
 
 	bool bCanBeActive = true;
 };
diff --git a/Engine/Plugins/Runtime/ComputeFramework/Source/ComputeFramework/Private/ComputeFramework.cpp b/Engine/Plugins/Runtime/ComputeFramework/Source/ComputeFramework/Private/ComputeFramework.cpp
--- a/Engine/Plugins/Runtime/ComputeFramework/Source/ComputeFramework/Private/ComputeFramework.cpp
+++ b/Engine/Plugins/Runtime/ComputeFramework/Source/ComputeFramework/Private/ComputeFramework.cpp
@@ -98,4 +98,18 @@ namespace ComputeFramework
 				});
 		}
 	}
+
+	void AbortWork(FSceneInterface const* InScene, UObject* InOwnerPointer)
+	{
+		FComputeFrameworkSystem* ComputeSystem = FComputeFrameworkModule::GetComputeSystem();
+		FComputeGraphTaskWorker* ComputeGraphWorker = ComputeSystem != nullptr ? ComputeSystem->GetComputeWorker(InScene) : nullptr;
+		if (ensure(ComputeGraphWorker))
+		{
+			ENQUEUE_RENDER_COMMAND(ComputeFrameworkAbortCommand)(
+				[ComputeGraphWorker, InOwnerPointer](FRHICommandListImmediate& RHICmdList)
+			{
+				ComputeGraphWorker->Abort(InOwnerPointer);
+			});
+		}
+	}
 }
diff --git a/Engine/Plugins/Runtime/ComputeFramework/Source/ComputeFramework/Private/ComputeFramework/ComputeGraphWorker.h b/Engine/Plugins/Runtime/ComputeFramework/Source/ComputeFramework/Private/ComputeFramework/ComputeGraphWorker.h
--- a/Engine/Plugins/Runtime/ComputeFramework/Source/ComputeFramework/Private/ComputeFramework/ComputeGraphWorker.h
+++ b/Engine/Plugins/Runtime/ComputeFramework/Source/ComputeFramework/Private/ComputeFramework/ComputeGraphWorker.h
@@ -27,7 +27,10 @@ public:
 		uint8 InGraphSortPriority,
 		FComputeGraphRenderProxy const* InGraphRenderProxy, 
 		TArray<FComputeDataProviderRenderProxy*> InDataProviderRenderProxies,
-		FSimpleDelegate InFallbackDelegate);
+		FSimpleDelegate InFallbackDelegate,
+		const UObject* OwnerPointer = nullptr);
+
+	void Abort(const UObject* OwnerPointer);
 
 	/** Has enqueued compute graph work. */
 	bool HasWork(FName InExecutionGroupName) const override;
@@ -44,6 +47,8 @@ private:
 	{
 		/** Name of owner object that invoked the graph. */
 		FName OwnerName;
+		/** Pointer to an owning UObject. */
+		const UObject* OwnerPointer = nullptr;
 		/** Priority used when sorting work. */
 		uint8 GraphSortPriority = 0;
 		/** Graph render proxy. */
diff --git a/Engine/Plugins/Runtime/ComputeFramework/Source/ComputeFramework/Private/ComputeGraphComponent.cpp b/Engine/Plugins/Runtime/ComputeFramework/Source/ComputeFramework/Private/ComputeGraphComponent.cpp
--- a/Engine/Plugins/Runtime/ComputeFramework/Source/ComputeFramework/Private/ComputeGraphComponent.cpp
+++ b/Engine/Plugins/Runtime/ComputeFramework/Source/ComputeFramework/Private/ComputeGraphComponent.cpp
@@ -1,7 +1,7 @@
 // Copyright Epic Games, Inc. All Rights Reserved.
 
 #include "ComputeFramework/ComputeGraphComponent.h"
-
+#include "ComputeFramework/ComputeFramework.h"
 #include "ComputeFramework/ComputeGraph.h"
 #include "ComputeWorkerInterface.h"
 #include "GameFramework/Actor.h"
@@ -49,6 +49,11 @@ void UComputeGraphComponent::TickComponent(float DeltaTime, enum ELevelTick Tick
 void UComputeGraphComponent::SendRenderDynamicData_Concurrent()
 {
 	Super::SendRenderDynamicData_Concurrent();
-	
-	ComputeGraphInstance.EnqueueWork(ComputeGraph, GetScene(), ComputeTaskExecutionGroup::EndOfFrameUpdate, GetOwner()->GetFName(), FSimpleDelegate());
+
+	ComputeGraphInstance.EnqueueWork(ComputeGraph, GetScene(), ComputeTaskExecutionGroup::EndOfFrameUpdate, GetOwner()->GetFName(), FSimpleDelegate(), this);
 }
+
+void UComputeGraphComponent::DestroyRenderState_Concurrent()
+{
+	ComputeFramework::AbortWork(GetScene(), this);
+}
\ No newline at end of file
diff --git a/Engine/Plugins/Runtime/ComputeFramework/Source/ComputeFramework/Private/ComputeGraphInstance.cpp b/Engine/Plugins/Runtime/ComputeFramework/Source/ComputeFramework/Private/ComputeGraphInstance.cpp
--- a/Engine/Plugins/Runtime/ComputeFramework/Source/ComputeFramework/Private/ComputeGraphInstance.cpp
+++ b/Engine/Plugins/Runtime/ComputeFramework/Source/ComputeFramework/Private/ComputeGraphInstance.cpp
@@ -26,7 +26,7 @@ void FComputeGraphInstance::DestroyDataProviders()
 	DataProviders.Reset();
 }
 
-bool FComputeGraphInstance::EnqueueWork(UComputeGraph* InComputeGraph, FSceneInterface const* InScene, FName InExecutionGroupName, FName InOwnerName, FSimpleDelegate InFallbackDelegate)
+bool FComputeGraphInstance::EnqueueWork(UComputeGraph* InComputeGraph, FSceneInterface const* InScene, FName InExecutionGroupName, FName InOwnerName, FSimpleDelegate InFallbackDelegate, UObject* InOwnerPointer)
 {
 	if (InComputeGraph == nullptr || InScene == nullptr)
 	{
@@ -71,10 +71,10 @@ bool FComputeGraphInstance::EnqueueWork(UComputeGraph* InComputeGraph, FSceneInt
 	}
 
 	ENQUEUE_RENDER_COMMAND(ComputeFrameworkEnqueueExecutionCommand)(
-		[ComputeGraphWorker, InExecutionGroupName, InOwnerName, SortPriority = GraphSortPriority, GraphRenderProxy, MovedDataProviderRenderProxies = MoveTemp(DataProviderRenderProxies), InFallbackDelegate](FRHICommandListImmediate& RHICmdList)
+		[ComputeGraphWorker, InExecutionGroupName, InOwnerName, SortPriority = GraphSortPriority, GraphRenderProxy, MovedDataProviderRenderProxies = MoveTemp(DataProviderRenderProxies), InFallbackDelegate, InOwnerPointer](FRHICommandListImmediate& RHICmdList)
 		{
 			// Compute graph scheduler will take ownership of the provider proxies.
-			ComputeGraphWorker->Enqueue(InExecutionGroupName, InOwnerName, SortPriority, GraphRenderProxy, MovedDataProviderRenderProxies, InFallbackDelegate);
+			ComputeGraphWorker->Enqueue(InExecutionGroupName, InOwnerName, SortPriority, GraphRenderProxy, MovedDataProviderRenderProxies, InFallbackDelegate, InOwnerPointer);
 		});
 
 	return true;
diff --git a/Engine/Plugins/Runtime/ComputeFramework/Source/ComputeFramework/Private/ComputeGraphWorker.cpp b/Engine/Plugins/Runtime/ComputeFramework/Source/ComputeFramework/Private/ComputeGraphWorker.cpp
--- a/Engine/Plugins/Runtime/ComputeFramework/Source/ComputeFramework/Private/ComputeGraphWorker.cpp
+++ b/Engine/Plugins/Runtime/ComputeFramework/Source/ComputeFramework/Private/ComputeGraphWorker.cpp
@@ -29,16 +29,38 @@ void FComputeGraphTaskWorker::Enqueue(
 	uint8 InGraphSortPriority,
 	FComputeGraphRenderProxy const* InGraphRenderProxy, 
 	TArray<FComputeDataProviderRenderProxy*> InDataProviderRenderProxies, 
-	FSimpleDelegate InFallbackDelegate)
+	FSimpleDelegate InFallbackDelegate,
+	const UObject* InOwnerPointer)
 {
 	FGraphInvocation& GraphInvocation = GraphInvocationsPerGroup.FindOrAdd(InExecutionGroupName).AddDefaulted_GetRef();
 	GraphInvocation.OwnerName = InOwnerName;
+	GraphInvocation.OwnerPointer = InOwnerPointer;
 	GraphInvocation.GraphSortPriority = InGraphSortPriority;
 	GraphInvocation.GraphRenderProxy = InGraphRenderProxy;
 	GraphInvocation.DataProviderRenderProxies = MoveTemp(InDataProviderRenderProxies);
 	GraphInvocation.FallbackDelegate = InFallbackDelegate;
 }
 
+void FComputeGraphTaskWorker::Abort(const UObject* InOwnerPointer)
+{
+	for (auto& Pair : GraphInvocationsPerGroup)
+	{
+		TArray<FGraphInvocation>& Invocations = Pair.Value;
+
+		for (int32 Index = 0; Index < Invocations.Num(); ++Index)
+		{
+			if (Invocations[Index].OwnerPointer == InOwnerPointer)
+			{
+				Invocations.RemoveAtSwap(Index);
+			}
+			else
+			{
+				Index++;
+			}
+		}
+	}
+}
+
 bool FComputeGraphTaskWorker::HasWork(FName InExecutionGroupName) const
 {
 	TArray<FGraphInvocation> const* GraphInvocations = GraphInvocationsPerGroup.Find(InExecutionGroupName);
diff --git a/Engine/Plugins/Runtime/ComputeFramework/Source/ComputeFramework/Public/ComputeFramework/ComputeFramework.h b/Engine/Plugins/Runtime/ComputeFramework/Source/ComputeFramework/Public/ComputeFramework/ComputeFramework.h
--- a/Engine/Plugins/Runtime/ComputeFramework/Source/ComputeFramework/Public/ComputeFramework/ComputeFramework.h
+++ b/Engine/Plugins/Runtime/ComputeFramework/Source/ComputeFramework/Public/ComputeFramework/ComputeFramework.h
@@ -31,4 +31,7 @@ namespace ComputeFramework
 
 	/** Flush any enqueued ComputeGraph work for a given execution group. */
 	COMPUTEFRAMEWORK_API void FlushWork(FSceneInterface const* InScene, FName InExecutionGroupName);
+
+	/** Aborts any queued work associated with the scene / execution group / owner pointer tuple. */
+	COMPUTEFRAMEWORK_API void AbortWork(FSceneInterface const* InScene, UObject* OwnerPointer);
 }
diff --git a/Engine/Plugins/Runtime/ComputeFramework/Source/ComputeFramework/Public/ComputeFramework/ComputeGraphComponent.h b/Engine/Plugins/Runtime/ComputeFramework/Source/ComputeFramework/Public/ComputeFramework/ComputeGraphComponent.h
--- a/Engine/Plugins/Runtime/ComputeFramework/Source/ComputeFramework/Public/ComputeFramework/ComputeGraphComponent.h
+++ b/Engine/Plugins/Runtime/ComputeFramework/Source/ComputeFramework/Public/ComputeFramework/ComputeGraphComponent.h
@@ -41,6 +41,7 @@ protected:
 	//~ Begin UActorComponent Interface
 	void TickComponent(float DeltaTime, enum ELevelTick TickType, FActorComponentTickFunction* ThisTickFunction) override;
 	void SendRenderDynamicData_Concurrent() override;
+	void DestroyRenderState_Concurrent() override;
 	bool ShouldCreateRenderState() const override { return true; }
 	//~ End UActorComponent Interface
 
diff --git a/Engine/Plugins/Runtime/ComputeFramework/Source/ComputeFramework/Public/ComputeFramework/ComputeGraphInstance.h b/Engine/Plugins/Runtime/ComputeFramework/Source/ComputeFramework/Public/ComputeFramework/ComputeGraphInstance.h
--- a/Engine/Plugins/Runtime/ComputeFramework/Source/ComputeFramework/Public/ComputeFramework/ComputeGraphInstance.h
+++ b/Engine/Plugins/Runtime/ComputeFramework/Source/ComputeFramework/Public/ComputeFramework/ComputeGraphInstance.h
@@ -38,7 +38,7 @@ public:
 	TArray< TObjectPtr<UComputeDataProvider> >& GetDataProviders() { return DataProviders; }
 
 	/** Enqueue the ComputeGraph work. */
-	bool EnqueueWork(UComputeGraph* InComputeGraph, FSceneInterface const* InScene, FName InExecutionGroupName, FName InOwnerName, FSimpleDelegate InFallbackDelegate);
+	bool EnqueueWork(UComputeGraph* InComputeGraph, FSceneInterface const* InScene, FName InExecutionGroupName, FName InOwnerName, FSimpleDelegate InFallbackDelegate, UObject* InOwnerPointer = nullptr);
 
 private:
 	/** The currently bound Data Provider objects. */
