Remove ensures to null check compute graph worker.

diff --git a/Engine/Plugins/Runtime/ComputeFramework/Source/ComputeFramework/Private/ComputeFramework.cpp b/Engine/Plugins/Runtime/ComputeFramework/Source/ComputeFramework/Private/ComputeFramework.cpp
--- a/Engine/Plugins/Runtime/ComputeFramework/Source/ComputeFramework/Private/ComputeFramework.cpp
+++ b/Engine/Plugins/Runtime/ComputeFramework/Source/ComputeFramework/Private/ComputeFramework.cpp
@@ -85,7 +85,7 @@ namespace ComputeFramework
 	{
 		FComputeFrameworkSystem* ComputeSystem = FComputeFrameworkModule::GetComputeSystem();
 		FComputeGraphTaskWorker* ComputeGraphWorker = ComputeSystem != nullptr ? ComputeSystem->GetComputeWorker(InScene) : nullptr;
-		if (ensure(ComputeGraphWorker))
+		if (ComputeGraphWorker)
 		{
 			UE::RenderCommandPipe::FSyncScope SyncScope;
 
@@ -103,7 +103,7 @@ namespace ComputeFramework
 	{
 		FComputeFrameworkSystem* ComputeSystem = FComputeFrameworkModule::GetComputeSystem();
 		FComputeGraphTaskWorker* ComputeGraphWorker = ComputeSystem != nullptr ? ComputeSystem->GetComputeWorker(InScene) : nullptr;
-		if (ensure(ComputeGraphWorker))
+		if (ComputeGraphWorker)
 		{
 			ENQUEUE_RENDER_COMMAND(ComputeFrameworkAbortCommand)(
 				[ComputeGraphWorker, InOwnerPointer](FRHICommandListImmediate& RHICmdList)
