ComputeFramework: fixed a typo when serialzing kernel resource set

diff --git a/Engine/Plugins/Runtime/ComputeFramework/Source/ComputeFramework/Private/ComputeGraph.cpp b/Engine/Plugins/Runtime/ComputeFramework/Source/ComputeFramework/Private/ComputeGraph.cpp
--- a/Engine/Plugins/Runtime/ComputeFramework/Source/ComputeFramework/Private/ComputeGraph.cpp
+++ b/Engine/Plugins/Runtime/ComputeFramework/Source/ComputeFramework/Private/ComputeGraph.cpp
@@ -945,7 +945,7 @@ void UComputeGraph::FComputeKernelResourceSet::Serialize(FArchive& Ar)
 					if (GMaxRHIShaderPlatform == ShaderMap->GetShaderPlatform())
 					{
 #if WITH_EDITORONLY_DATA
-						KernelResourcesByFeatureLevel[GMaxRHIShaderPlatform] = MoveTemp(Resource);
+						KernelResourcesByFeatureLevel[GMaxRHIFeatureLevel] = MoveTemp(Resource);
 #else
 						KernelResource = MoveTemp(Resource);
 #endif
