Fix compute graph Abort skipping enqueued items, leading to potential crashes when stale items are left in the invocation array.
Also changed use of RemoveAtSwap() to RemoveAt() so that ordering of enqueued items is respected when executing.

diff --git a/Engine/Plugins/Runtime/ComputeFramework/Source/ComputeFramework/Private/ComputeGraphWorker.cpp b/Engine/Plugins/Runtime/ComputeFramework/Source/ComputeFramework/Private/ComputeGraphWorker.cpp
--- a/Engine/Plugins/Runtime/ComputeFramework/Source/ComputeFramework/Private/ComputeGraphWorker.cpp
+++ b/Engine/Plugins/Runtime/ComputeFramework/Source/ComputeFramework/Private/ComputeGraphWorker.cpp
@@ -47,15 +47,11 @@ void FComputeGraphTaskWorker::Abort(const UObject* InOwnerPointer)
 	{
 		TArray<FGraphInvocation>& Invocations = Pair.Value;
 
-		for (int32 Index = 0; Index < Invocations.Num(); ++Index)
+		for (int32 Index = Invocations.Num() - 1; Index >= 0; Index--)
 		{
 			if (Invocations[Index].OwnerPointer == InOwnerPointer)
 			{
-				Invocations.RemoveAtSwap(Index);
-			}
-			else
-			{
-				Index++;
+				Invocations.RemoveAt(Index, 1, EAllowShrinking::No);
 			}
 		}
 	}
