[Compute Framework & Deformer Graph] fix incorrect function shim shader code generated when a paramter pin connects to a resource pin, the index arugment should be ignored

diff --git a/Engine/Plugins/Runtime/ComputeFramework/Source/ComputeFramework/Private/ComputeGraph.cpp b/Engine/Plugins/Runtime/ComputeFramework/Source/ComputeFramework/Private/ComputeGraph.cpp
--- a/Engine/Plugins/Runtime/ComputeFramework/Source/ComputeFramework/Private/ComputeGraph.cpp
+++ b/Engine/Plugins/Runtime/ComputeFramework/Source/ComputeFramework/Private/ComputeGraph.cpp
@@ -405,7 +405,8 @@ namespace
 	void GetFunctionShimHLSL(FShaderFunctionDefinition const& FnImpl, FShaderFunctionDefinition const& FnWrap, TCHAR const* UID, TCHAR const *WrapNameOverride, TCHAR const* Namespace, FString& InOutHLSL)
 	{
 		const bool bHasReturn = FnWrap.bHasReturnType;
-		const int32 NumParams = FnWrap.ParamTypes.Num();
+		const int32 NumImplParams = FnImpl.ParamTypes.Num();
+		const int32 NumWrapParams = FnWrap.ParamTypes.Num();
 
 		TStringBuilder<512> StringBuilder;
 
@@ -420,11 +421,11 @@ namespace
 		StringBuilder.Append(WrapNameOverride ? WrapNameOverride : *FnWrap.Name);
 		StringBuilder.Append(TEXT("("));
 		
-		for (int32 ParameterIndex = bHasReturn ? 1 : 0; ParameterIndex < NumParams; ++ParameterIndex)
+		for (int32 ParameterIndex = bHasReturn ? 1 : 0; ParameterIndex < NumWrapParams; ++ParameterIndex)
 		{
 			StringBuilder.Append(*FnWrap.ParamTypes[ParameterIndex].TypeDeclaration);
 			StringBuilder.Appendf(TEXT(" P%d"), ParameterIndex);
-			StringBuilder.Append((ParameterIndex < NumParams - 1) ? TEXT(", ") : TEXT(""));
+			StringBuilder.Append((ParameterIndex < NumWrapParams - 1) ? TEXT(", ") : TEXT(""));
 		}
 
 		StringBuilder.Append(TEXT(") { "));
@@ -432,10 +433,10 @@ namespace
 		StringBuilder.Append(*FnImpl.Name).Append(TEXT("_")).Append(UID);
 		StringBuilder.Append(TEXT("("));
 
-		for (int32 ParameterIndex = bHasReturn ? 1 : 0; ParameterIndex < NumParams; ++ParameterIndex)
+		for (int32 ParameterIndex = bHasReturn ? 1 : 0; ParameterIndex < NumImplParams; ++ParameterIndex)
 		{
 			StringBuilder.Appendf(TEXT("P%d"), ParameterIndex);
-			StringBuilder.Append((ParameterIndex < NumParams - 1) ? TEXT(", ") : TEXT(""));
+			StringBuilder.Append((ParameterIndex < NumImplParams - 1) ? TEXT(", ") : TEXT(""));
 		}
 
 		StringBuilder.Append(TEXT(");"));
