[Compute Framework] fix incorrect parameter indexing in function shim hlsl for data interface function with a return value.

diff --git a/Engine/Plugins/Runtime/ComputeFramework/Source/ComputeFramework/Private/ComputeGraph.cpp b/Engine/Plugins/Runtime/ComputeFramework/Source/ComputeFramework/Private/ComputeGraph.cpp
--- a/Engine/Plugins/Runtime/ComputeFramework/Source/ComputeFramework/Private/ComputeGraph.cpp
+++ b/Engine/Plugins/Runtime/ComputeFramework/Source/ComputeFramework/Private/ComputeGraph.cpp
@@ -404,8 +404,10 @@ namespace
 	/** Add HLSL code to implement an external function. */
 	void GetFunctionShimHLSL(FShaderFunctionDefinition const& FnImpl, FShaderFunctionDefinition const& FnWrap, TCHAR const* UID, TCHAR const *WrapNameOverride, TCHAR const* Namespace, FString& InOutHLSL)
 	{
-		const bool bHasReturn = FnWrap.bHasReturnType;
+		const bool bHasReturnImpl = FnImpl.bHasReturnType;
 		const int32 NumImplParams = FnImpl.ParamTypes.Num();
+		
+		const bool bHasReturnWrap = FnWrap.bHasReturnType;
 		const int32 NumWrapParams = FnWrap.ParamTypes.Num();
 
 		TStringBuilder<512> StringBuilder;
@@ -416,12 +418,12 @@ namespace
 			StringBuilder.Append(Namespace);
 			StringBuilder.Append(TEXT(" { "));
 		}
-		StringBuilder.Append(bHasReturn ? *FnWrap.ParamTypes[0].TypeDeclaration : TEXT("void"));
+		StringBuilder.Append(bHasReturnWrap ? *FnWrap.ParamTypes[0].TypeDeclaration : TEXT("void"));
 		StringBuilder.Append(TEXT(" "));
 		StringBuilder.Append(WrapNameOverride ? WrapNameOverride : *FnWrap.Name);
 		StringBuilder.Append(TEXT("("));
 		
-		for (int32 ParameterIndex = bHasReturn ? 1 : 0; ParameterIndex < NumWrapParams; ++ParameterIndex)
+		for (int32 ParameterIndex = bHasReturnWrap ? 1 : 0; ParameterIndex < NumWrapParams; ++ParameterIndex)
 		{
 			StringBuilder.Append(*FnWrap.ParamTypes[ParameterIndex].TypeDeclaration);
 			StringBuilder.Appendf(TEXT(" P%d"), ParameterIndex);
@@ -429,11 +431,11 @@ namespace
 		}
 
 		StringBuilder.Append(TEXT(") { "));
-		StringBuilder.Append(bHasReturn ? TEXT("return ") : TEXT(""));
+		StringBuilder.Append(bHasReturnWrap ? TEXT("return ") : TEXT(""));
 		StringBuilder.Append(*FnImpl.Name).Append(TEXT("_")).Append(UID);
 		StringBuilder.Append(TEXT("("));
 
-		for (int32 ParameterIndex = bHasReturn ? 1 : 0; ParameterIndex < NumImplParams; ++ParameterIndex)
+		for (int32 ParameterIndex = bHasReturnImpl ? 1 : 0; ParameterIndex < NumImplParams; ++ParameterIndex)
 		{
 			StringBuilder.Appendf(TEXT("P%d"), ParameterIndex);
 			StringBuilder.Append((ParameterIndex < NumImplParams - 1) ? TEXT(", ") : TEXT(""));
