[Compute Graph] Additional fix to generated function shim hlsl when the Function Implementation has fewer input param than the wrap function

diff --git a/Engine/Plugins/Runtime/ComputeFramework/Source/ComputeFramework/Private/ComputeGraph.cpp b/Engine/Plugins/Runtime/ComputeFramework/Source/ComputeFramework/Private/ComputeGraph.cpp
--- a/Engine/Plugins/Runtime/ComputeFramework/Source/ComputeFramework/Private/ComputeGraph.cpp
+++ b/Engine/Plugins/Runtime/ComputeFramework/Source/ComputeFramework/Private/ComputeGraph.cpp
@@ -406,6 +406,7 @@ namespace
 	{
 		const bool bHasReturnImpl = FnImpl.bHasReturnType;
 		const int32 NumImplParams = FnImpl.ParamTypes.Num();
+		const int32 NumImplInputParams = bHasReturnImpl ? NumImplParams - 1 : NumImplParams;
 		
 		const bool bHasReturnWrap = FnWrap.bHasReturnType;
 		const int32 NumWrapParams = FnWrap.ParamTypes.Num();
@@ -435,10 +436,22 @@ namespace
 		StringBuilder.Append(*FnImpl.Name).Append(TEXT("_")).Append(UID);
 		StringBuilder.Append(TEXT("("));
 
-		for (int32 ParameterIndex = bHasReturnImpl ? 1 : 0; ParameterIndex < NumImplParams; ++ParameterIndex)
+		// There are cases where the impl will have fewer input params than the wrap, additional wrap params should be skipped
+		// Example: when a parameter pin connects to a resource pin
+		// void Wrap(uint P0, uint P1, ...) { Impl(); } // Impl has no input param
+		// SomeType Wrap(uint P1, uint P2, ...) { return Impl(P1); } // Impl has 1 input param
+		int32 NumImplInputParamsUsed = 0;
+		for (int32 WrapParameterIndex = bHasReturnWrap? 1 : 0; WrapParameterIndex < NumWrapParams; ++WrapParameterIndex)
 		{
-			StringBuilder.Appendf(TEXT("P%d"), ParameterIndex);
-			StringBuilder.Append((ParameterIndex < NumImplParams - 1) ? TEXT(", ") : TEXT(""));
+			if (NumImplInputParamsUsed >= NumImplInputParams)
+			{
+				break;
+			}
+			// prepend a comma if we are not the first param for the impl
+			StringBuilder.Append((NumImplInputParamsUsed != 0) ? TEXT(", ") : TEXT(""));
+			StringBuilder.Appendf(TEXT("P%d"), WrapParameterIndex);
+			
+			NumImplInputParamsUsed++;
 		}
 
 		StringBuilder.Append(TEXT(");"));
