[Compute Framework] added newline before kernel includes to avoid preprocessor failing to process those includes if a data interface's shader file does not end with an extra empty line

diff --git a/Engine/Plugins/Runtime/ComputeFramework/Source/ComputeFramework/Private/ComputeGraph.cpp b/Engine/Plugins/Runtime/ComputeFramework/Source/ComputeFramework/Private/ComputeGraph.cpp
--- a/Engine/Plugins/Runtime/ComputeFramework/Source/ComputeFramework/Private/ComputeGraph.cpp
+++ b/Engine/Plugins/Runtime/ComputeFramework/Source/ComputeFramework/Private/ComputeGraph.cpp
@@ -517,7 +517,7 @@ FString UComputeGraph::BuildKernelSource(
 	// Add virtual source includes from the additional sources.
 	for (TPair<FString, FString> const& AdditionalSource : InAdditionalSources)
 	{
-		HLSL += FString::Printf(TEXT("#include \"%s\"\n"), *AdditionalSource.Key);
+		HLSL += FString::Printf(TEXT("\n#include \"%s\"\n"), *AdditionalSource.Key);
 
 		// Accumulate the source HLSL to the local hash state.
 		HashState.UpdateWithString(*AdditionalSource.Value, AdditionalSource.Value.Len());
@@ -553,7 +553,7 @@ FString UComputeGraph::BuildKernelSource(
 			{
 				// The generated path has a magic unique prefix which the compilation manager knows to strip before resolving errors.
 				FString MagicVirtualPath = FString::Printf(TEXT("/Engine/Generated/DataInterface/%s%s"), *NamePrefix, ShaderVirtualPath);
-				HLSL += FString::Printf(TEXT("#include \"%s\"\n"), *MagicVirtualPath);
+				HLSL += FString::Printf(TEXT("\n#include \"%s\"\n"), *MagicVirtualPath);
 				FString DataInterfaceHLSL;
 				DataInterface->GetHLSL(DataInterfaceHLSL, NamePrefix);
 				OutGeneratedSources.Add(MagicVirtualPath, DataInterfaceHLSL);
