[Compute Framework] Add readback support as API through data interface

Readbacks are executed via FRHIGPUBufferReadback. FComputeGraphTaskWorker creates the readbacks and polls active readbacks each frame.

diff --git a/Engine/Plugins/Runtime/ComputeFramework/Source/ComputeFramework/Private/ComputeFramework/ComputeGraphRenderProxy.h b/Engine/Plugins/Runtime/ComputeFramework/Source/ComputeFramework/Private/ComputeFramework/ComputeGraphRenderProxy.h
--- a/Engine/Plugins/Runtime/ComputeFramework/Source/ComputeFramework/Private/ComputeFramework/ComputeGraphRenderProxy.h
+++ b/Engine/Plugins/Runtime/ComputeFramework/Source/ComputeFramework/Private/ComputeFramework/ComputeGraphRenderProxy.h
@@ -30,11 +30,13 @@ public:
 		FShaderParametersMetadata* ShaderParameterMetadata = nullptr;
 		/** Array of indices into the full graph data provider array. Contains only the indices to data providers that this kernel references. */
 		TArray<int32> BoundProviderIndices;
+		/** Array of indices for data providers that should trigger a readback. */
+		TArray<int32> ReadbackProviderIndices;
 		/** Same size as BoundProviderIndices, non-primary data providers will be forced to present a full view to its data (unified)*/
 		TArray<bool> BoundProviderIsPrimary;
 		/** The index of the special execution data provider in the full graph data provider array. */
 		int32 ExecutionProviderIndex = -1;
-		/** Whether the kernel can combine multiple supports combinined sub-invocations. */
+		/** Whether the kernel can combine multiple supports combined sub-invocations. */
 		bool bSupportsUnifiedDispatch = false;
 	};
 
diff --git a/Engine/Plugins/Runtime/ComputeFramework/Source/ComputeFramework/Private/ComputeFramework/ComputeGraphWorker.h b/Engine/Plugins/Runtime/ComputeFramework/Source/ComputeFramework/Private/ComputeFramework/ComputeGraphWorker.h
--- a/Engine/Plugins/Runtime/ComputeFramework/Source/ComputeFramework/Private/ComputeFramework/ComputeGraphWorker.h
+++ b/Engine/Plugins/Runtime/ComputeFramework/Source/ComputeFramework/Private/ComputeFramework/ComputeGraphWorker.h
@@ -41,6 +41,10 @@ public:
 		FName InExecutionGroupName,
 		ERHIFeatureLevel::Type InFeatureLevel ) override;
 
+private:
+	/** Poll active readbacks and execute callbacks for ready ones. */
+	void UpdateReadbacks();
+
 private:
 	/** Description of each graph that is enqueued. */
 	struct FGraphInvocation
@@ -97,4 +101,18 @@ private:
 	TArray<TShaderRef<FComputeKernelShader>> Shaders;
 	TArray<int32> PermutationIds;
 	TArray<FIntVector> ThreadCounts;
+	TArray<FComputeDataProviderRenderProxy::FReadbackData> ReadbackDatas;
+
+	/** Data for tracking active asynchronous readbacks. */
+	struct FAsyncReadback
+	{
+		~FAsyncReadback();
+
+		class FRHIGPUBufferReadback* Readback;
+		uint32 NumBytes;
+		const UObject* OwnerPointer = nullptr;
+		FComputeDataProviderRenderProxy::FReadbackCallback OnDataAvailable;
+	};
+
+	TArray<FAsyncReadback> ActiveAsyncReadbacks;
 };
diff --git a/Engine/Plugins/Runtime/ComputeFramework/Source/ComputeFramework/Private/ComputeGraph.cpp b/Engine/Plugins/Runtime/ComputeFramework/Source/ComputeFramework/Private/ComputeGraph.cpp
--- a/Engine/Plugins/Runtime/ComputeFramework/Source/ComputeFramework/Private/ComputeGraph.cpp
+++ b/Engine/Plugins/Runtime/ComputeFramework/Source/ComputeFramework/Private/ComputeGraph.cpp
@@ -373,6 +373,12 @@ FComputeGraphRenderProxy* UComputeGraph::CreateRenderProxy() const
 							{
 								Invocation.bSupportsUnifiedDispatch &= DataInterface->CanSupportUnifiedDispatch();
 							}
+
+							// If the data interface is requesting readback and is a kernel output, register it.
+							if (!GraphEdge.bKernelInput && DataInterface->GetRequiresReadback())
+							{
+								Invocation.ReadbackProviderIndices.AddUnique(GraphEdge.DataInterfaceIndex);
+							}
 						}
 					}
 				}
diff --git a/Engine/Plugins/Runtime/ComputeFramework/Source/ComputeFramework/Private/ComputeGraphWorker.cpp b/Engine/Plugins/Runtime/ComputeFramework/Source/ComputeFramework/Private/ComputeGraphWorker.cpp
--- a/Engine/Plugins/Runtime/ComputeFramework/Source/ComputeFramework/Private/ComputeGraphWorker.cpp
+++ b/Engine/Plugins/Runtime/ComputeFramework/Source/ComputeFramework/Private/ComputeGraphWorker.cpp
@@ -12,6 +12,7 @@
 #include "RenderGraphBuilder.h"
 #include "RenderGraphUtils.h"
 #include "RHIDefinitions.h"
+#include "RHIGPUReadback.h"
 
 DECLARE_GPU_STAT_NAMED(ComputeFramework_ExecuteBatches, TEXT("ComputeFramework::ExecuteBatches"));
 
@@ -55,10 +56,23 @@ void FComputeGraphTaskWorker::Abort(const UObject* InOwnerPointer)
 			}
 		}
 	}
+
+	for (int I = ActiveAsyncReadbacks.Num() - 1; I >= 0; --I)
+	{
+		if (ActiveAsyncReadbacks[I].OwnerPointer == InOwnerPointer)
+		{
+			ActiveAsyncReadbacks.RemoveAtSwap(I);
+		}
+	}
 }
 
 bool FComputeGraphTaskWorker::HasWork(FName InExecutionGroupName) const
 {
+	if (!ActiveAsyncReadbacks.IsEmpty())
+	{
+		return true;
+	}
+
 	TArray<FGraphInvocation> const* GraphInvocations = GraphInvocationsPerGroup.Find(InExecutionGroupName);
 	return GraphInvocations != nullptr && GraphInvocations->Num();
 }
@@ -69,6 +83,12 @@ void FComputeGraphTaskWorker::SubmitWork(FRDGBuilder& GraphBuilder, FName InExec
 	RDG_EVENT_SCOPE(GraphBuilder, "ComputeFramework::ExecuteBatches");
 	RDG_GPU_STAT_SCOPE(GraphBuilder, ComputeFramework_ExecuteBatches);
 
+	// Currently poll readbacks once at end of frame.
+	if (InExecutionGroupName == ComputeTaskExecutionGroup::EndOfFrameUpdate)
+	{
+		UpdateReadbacks();
+	}
+
 	// Reset our scratch memory arrays.
 	SubmitDescs.Reset();
 	Shaders.Reset();
@@ -263,6 +283,35 @@ void FComputeGraphTaskWorker::SubmitWork(FRDGBuilder& GraphBuilder, FName InExec
 				GroupCount
 			);
 		}
+
+		// Enqueue readbacks.
+		for (int32 DataProviderIndex : KernelInvocation.ReadbackProviderIndices)
+		{
+			if (FComputeDataProviderRenderProxy* DataProviderProxy = GraphInvocation.DataProviderRenderProxies[DataProviderIndex])
+			{
+				ReadbackDatas.Reset();
+				DataProviderProxy->GetReadbackData(ReadbackDatas);
+
+				for (const FComputeDataProviderRenderProxy::FReadbackData& ReadbackData : ReadbackDatas)
+				{
+					if (!ReadbackData.Buffer || ReadbackData.NumBytes == 0 || !ReadbackData.ReadbackCallback_RenderThread)
+					{
+						continue;
+					}
+
+					FRHIGPUBufferReadback* ReadbackRequest = new FRHIGPUBufferReadback(TEXT("ComputeFrameworkBuffer"));
+					check(ReadbackRequest);
+
+					FAsyncReadback& Readback = ActiveAsyncReadbacks.Emplace_GetRef();
+					Readback.Readback = ReadbackRequest;
+					Readback.NumBytes = ReadbackData.NumBytes;
+					Readback.OwnerPointer = GraphInvocation.OwnerPointer;
+					Readback.OnDataAvailable = *ReadbackData.ReadbackCallback_RenderThread;
+
+					AddEnqueueCopyPass(GraphBuilder, Readback.Readback, ReadbackData.Buffer, Readback.NumBytes);
+				}
+			}
+		}
 	}
 
 	// Release any graph resources at the end of graph execution.
@@ -275,6 +324,29 @@ void FComputeGraphTaskWorker::SubmitWork(FRDGBuilder& GraphBuilder, FName InExec
 		});
 }
 
+void FComputeGraphTaskWorker::UpdateReadbacks()
+{
+	TRACE_CPUPROFILER_EVENT_SCOPE(ComputeFramework::UpdateReadbacks);
+
+	for (int I = ActiveAsyncReadbacks.Num() - 1; I >= 0; --I)
+	{
+		FRHIGPUBufferReadback* Request = ActiveAsyncReadbacks[I].Readback;
+		const uint32 NumBytes = ActiveAsyncReadbacks[I].NumBytes;
+
+		if (ensure(Request) && Request->IsReady())
+		{
+			void* ReadbackData = Request->Lock(NumBytes);
+
+			if (ensure(ReadbackData))
+			{
+				ActiveAsyncReadbacks[I].OnDataAvailable(ReadbackData, NumBytes);
+			}
+
+			ActiveAsyncReadbacks.RemoveAtSwap(I);
+		}
+	}
+}
+
 FComputeGraphTaskWorker::FGraphInvocation::~FGraphInvocation()
 {
 	// DataProviderRenderProxy objects are created per frame and destroyed here after render work has been submitted.
@@ -284,3 +356,12 @@ FComputeGraphTaskWorker::FGraphInvocation::~FGraphInvocation()
 		delete DataProvider;
 	}
 }
+
+FComputeGraphTaskWorker::FAsyncReadback::~FAsyncReadback()
+{
+	if (ensure(Readback))
+	{
+		delete Readback;
+		Readback = nullptr;
+	}
+}
diff --git a/Engine/Plugins/Runtime/ComputeFramework/Source/ComputeFramework/Public/ComputeFramework/ComputeDataInterface.h b/Engine/Plugins/Runtime/ComputeFramework/Source/ComputeFramework/Public/ComputeFramework/ComputeDataInterface.h
--- a/Engine/Plugins/Runtime/ComputeFramework/Source/ComputeFramework/Public/ComputeFramework/ComputeDataInterface.h
+++ b/Engine/Plugins/Runtime/ComputeFramework/Source/ComputeFramework/Public/ComputeFramework/ComputeDataInterface.h
@@ -47,6 +47,8 @@ public:
 	virtual void GetStructDeclarations(TSet<FString>& OutStructsSeen, TArray<FString>& OutStructs) const {}
 	/** Gather the shader code for this data provider. */
 	virtual void GetHLSL(FString& OutHLSL, FString const& InDataInterfaceName) const {}
+	/** Whether a GPU->CPU readback should be scheduled after this data is used as a kernel output. */
+	virtual bool GetRequiresReadback() const { return false; }
 	/** Get class of UObject required to instantiate a UComputeDataProvider from this interface. */
 	virtual UClass* GetBindingType() const { return nullptr; }
 	/** Instantiate an associated UComputeDataProvider. */
diff --git a/Engine/Plugins/Runtime/ComputeFramework/Source/ComputeFramework/Public/ComputeFramework/ComputeDataProvider.h b/Engine/Plugins/Runtime/ComputeFramework/Source/ComputeFramework/Public/ComputeFramework/ComputeDataProvider.h
--- a/Engine/Plugins/Runtime/ComputeFramework/Source/ComputeFramework/Public/ComputeFramework/ComputeDataProvider.h
+++ b/Engine/Plugins/Runtime/ComputeFramework/Source/ComputeFramework/Public/ComputeFramework/ComputeDataProvider.h
@@ -3,6 +3,7 @@
 #pragma once
 
 #include "CoreMinimal.h"
+#include "RenderGraphFwd.h"
 #include "Containers/StridedView.h"
 #include "ComputeDataProvider.generated.h"
 
@@ -107,6 +108,22 @@ public:
 	/** Collect parameter data required to dispatch work. */
 	virtual void GatherDispatchData(FDispatchData const& InDispatchData) {}
 
+	using FReadbackCallback = TFunction<void(const void* InData, int InNumBytes)>;
+	struct FReadbackData
+	{
+		/** The buffer to be read back. */
+		FRDGBufferRef Buffer;
+
+		/** The number of bytes to read back. */
+		uint32 NumBytes;
+
+		/** Callback to execute once data is ready for CPU consumption. */
+		const FReadbackCallback* ReadbackCallback_RenderThread;
+	};
+
+	/** Data for any readbacks that should be performed. */
+	virtual void GetReadbackData(TArray<FReadbackData>& OutReadbackData) const {}
+
 protected:
 	/** Helper for making an FStridedView over the FDispatchData. */
 	template <typename ElementType>
