[Compute Framework] Added missing member initialisation

diff --git a/Engine/Plugins/Runtime/ComputeFramework/Source/ComputeFramework/Private/ComputeFramework/ComputeGraphWorker.h b/Engine/Plugins/Runtime/ComputeFramework/Source/ComputeFramework/Private/ComputeFramework/ComputeGraphWorker.h
--- a/Engine/Plugins/Runtime/ComputeFramework/Source/ComputeFramework/Private/ComputeFramework/ComputeGraphWorker.h
+++ b/Engine/Plugins/Runtime/ComputeFramework/Source/ComputeFramework/Private/ComputeFramework/ComputeGraphWorker.h
@@ -108,8 +108,8 @@ private:
 	{
 		~FAsyncReadback();
 
-		class FRHIGPUBufferReadback* Readback;
-		uint32 NumBytes;
+		class FRHIGPUBufferReadback* Readback = nullptr;
+		uint32 NumBytes = 0;
 		const UObject* OwnerPointer = nullptr;
 		FComputeDataProviderRenderProxy::FReadbackCallback OnDataAvailable;
 	};
