[ShaderPreprocessor] Parse line directives. Apply line number and filename to parser state so that subsequent errors will have line numbers with offsets applied from line directives.
[ComputeFramework] Treat preprocessor issues as warnings if compile succeeded, otherwise errors, which is the standard rule.
[PCG] Adapt error summary to whatever line information is available in the error (if any).

#rnx

diff --git a/Engine/Plugins/PCG/Source/PCG/Private/Compute/Elements/PCGComputeGraphElement.cpp b/Engine/Plugins/PCG/Source/PCG/Private/Compute/Elements/PCGComputeGraphElement.cpp
--- a/Engine/Plugins/PCG/Source/PCG/Private/Compute/Elements/PCGComputeGraphElement.cpp
+++ b/Engine/Plugins/PCG/Source/PCG/Private/Compute/Elements/PCGComputeGraphElement.cpp
@@ -189,7 +189,8 @@ bool FPCGComputeGraphElement::ExecuteInternal(FPCGContext* InContext) const
 		{
 			for (const FComputeKernelCompileMessage& Message : NodeAndCompileMessages.Get<1>())
 			{
-				// Currently failure messages can come through with Log severity, pattern match for now.
+				// Some error messages were getting lost, and we were only getting the final 'failed' message. Treat this as failure and report for now.
+				// TODO: Revert the 'failed' part once we're happy all relevant issues are bubbling up.
 				if (Message.Type == FComputeKernelCompileMessage::EMessageType::Error || Message.Text.Contains(TEXT("failed"), ESearchCase::IgnoreCase))
 				{
 					return true;
@@ -304,7 +305,9 @@ void FPCGComputeGraphElement::LogCompilationMessages(FPCGComputeGraphContext* In
 				}
 				else if (Message.Text.Contains(TEXT("failed"), ESearchCase::IgnoreCase))
 				{
-					// Currently failure messages can come through with Log verbosity.
+					// Some error messages were getting lost, and we were only getting the final 'failed' message.
+					// Treat this as failure and report for now.
+					// TODO: Revert this once we're happy all relevant issues are bubbling up.
 					Verbosity = ELogVerbosity::Error;
 				}
 
@@ -315,18 +318,25 @@ void FPCGComputeGraphElement::LogCompilationMessages(FPCGComputeGraphContext* In
 						FPCGStack StackWithNode = *InContext->Stack;
 						StackWithNode.PushFrame(NodeAndCompileMessages.Get<0>().ResolveObjectPtr());
 
+						FText LogText;
+
 						if (Message.Line != INDEX_NONE)
 						{
-							Subsystem->GetNodeVisualLogsMutable().Log(StackWithNode, Verbosity, FText::Format(
-								LOCTEXT("ErrorWithLineFormat", "[{0},{1}] {2}"),
-								Message.Line,
-								Message.ColumnStart,
-								FText::FromString(Message.Text)));
+							if (Message.ColumnStart != INDEX_NONE)
+							{
+								LogText = FText::Format(LOCTEXT("ErrorWithLineColFormat", "[{0},{1}] {2}"), Message.Line, Message.ColumnStart, FText::FromString(Message.Text));
+							}
+							else
+							{
+								LogText = FText::Format(LOCTEXT("ErrorWithLineFormat", "[{0}] {1}"), Message.Line, FText::FromString(Message.Text));
+							}
 						}
 						else
 						{
-							Subsystem->GetNodeVisualLogsMutable().Log(StackWithNode, Verbosity, FText::FromString(Message.Text));
+							LogText = FText::FromString(Message.Text);
 						}
+
+						Subsystem->GetNodeVisualLogsMutable().Log(StackWithNode, Verbosity, LogText);
 					}
 				}
 			}
diff --git a/Engine/Plugins/PCG/Source/PCG/Private/Compute/Elements/PCGCustomHLSL.cpp b/Engine/Plugins/PCG/Source/PCG/Private/Compute/Elements/PCGCustomHLSL.cpp
--- a/Engine/Plugins/PCG/Source/PCG/Private/Compute/Elements/PCGCustomHLSL.cpp
+++ b/Engine/Plugins/PCG/Source/PCG/Private/Compute/Elements/PCGCustomHLSL.cpp
@@ -1552,7 +1552,7 @@ FString UPCGCustomHLSLSettings::GetCookedKernelSource(const TMap<FPCGKernelAttri
 			"	if (ThreadIndex >= GetNumThreads().x) return;\n"
 			"%s\n" // HeaderWriters
 			"%s\n" // KernelSpecificPreamble
-			"#line 0 \"%s\"\n\n" // ShaderPathName
+			"#line 0 \"%s\"\n" // ShaderPathName
 			"%s\n" // Source
 			"}\n"),
 			*Includes, *Functions, *KernelFunc, *UnWrappedDispatchThreadId, *HeaderWriters, *KernelSpecificPreamble, *ShaderPathName, *Source);
diff --git a/Engine/Plugins/Runtime/ComputeFramework/Source/ComputeFramework/Private/ComputeKernelShaderCompilationManager.cpp b/Engine/Plugins/Runtime/ComputeFramework/Source/ComputeFramework/Private/ComputeKernelShaderCompilationManager.cpp
--- a/Engine/Plugins/Runtime/ComputeFramework/Source/ComputeFramework/Private/ComputeKernelShaderCompilationManager.cpp
+++ b/Engine/Plugins/Runtime/ComputeFramework/Source/ComputeFramework/Private/ComputeKernelShaderCompilationManager.cpp
@@ -105,25 +105,10 @@ void FComputeKernelShaderCompilationManager::ProcessAsyncResults()
 	}
 }
 
-static bool ParseShaderCompilerError(FShaderCompilerError const& InError, FComputeKernelCompileMessage& OutMessage)
+static bool ParseShaderCompilerError(FShaderCompilerError const& InError, bool bInCompilationSucceeded, FComputeKernelCompileMessage& OutMessage)
 {
 	FShaderCompilerError Error = InError;
 
-	// We ignore error messages that don't have a line information.
-	FString Line, Column;
-	if (!Error.HasLineMarker())
-	{
-		return false;
-	}
-	if (!Error.ErrorLineString.Split(TEXT(","), &Line, &Column))
-	{
-		return false;
-	}
-	if (!Line.IsNumeric() || !Column.IsNumeric())
-	{
-		return false;
-	}
-
 	if (Error.StrippedErrorMessage.RemoveFromStart(TEXT("error: ")))
 	{
 		OutMessage.Type = FComputeKernelCompileMessage::EMessageType::Error;
@@ -136,6 +121,11 @@ static bool ParseShaderCompilerError(FShaderCompilerError const& InError, FCompu
 	{
 		OutMessage.Type = FComputeKernelCompileMessage::EMessageType::Info;
 	}
+	else
+	{
+		// General rule for preprocessing errors - if compilation succeeded errors are warnings, otherwise errors.
+		OutMessage.Type = bInCompilationSucceeded ? FComputeKernelCompileMessage::EMessageType::Warning : FComputeKernelCompileMessage::EMessageType::Error;
+	}
 
 	OutMessage.Text = Error.StrippedErrorMessage;
 	OutMessage.VirtualFilePath = Error.ErrorVirtualFilePath;
@@ -152,13 +142,41 @@ static bool ParseShaderCompilerError(FShaderCompilerError const& InError, FCompu
 		OutMessage.RealFilePath = GetShaderSourceFilePath(OutMessage.VirtualFilePath);
 	}
 
-	LexFromString(OutMessage.Line, *Line);
-	LexFromString(OutMessage.ColumnStart, *Column);
-	OutMessage.ColumnEnd = OutMessage.ColumnStart;
-	
-	for (TCHAR Character : Error.HighlightedLineMarker)
+	// Populate line and column number strings if available.
+	FString Line, Column;
+	// Check for "line,col" format first.
+	if (!Error.ErrorLineString.Split(TEXT(","), &Line, &Column))
 	{
-		OutMessage.ColumnEnd += (Character == TEXT('~')) ? 1 : 0;
+		// Fall back to "line" which is logged by preprocessor errors.
+		if (Error.ErrorLineString.IsNumeric())
+		{
+			Line = Error.ErrorLineString;
+			Column.Empty();
+		}
+		else
+		{
+			Line.Empty();
+			Column.Empty();
+		}
+	}
+
+	if (Line.IsNumeric())
+	{
+		LexFromString(OutMessage.Line, *Line);
+	}
+
+	if (Column.IsNumeric())
+	{
+		LexFromString(OutMessage.ColumnStart, *Column);
+		OutMessage.ColumnEnd = OutMessage.ColumnStart;
+
+		if (Error.HasLineMarker())
+		{
+			for (TCHAR Character : Error.HighlightedLineMarker)
+			{
+				OutMessage.ColumnEnd += (Character == TEXT('~')) ? 1 : 0;
+			}
+		}
 	}
 
 	return true;
@@ -253,7 +271,7 @@ void FComputeKernelShaderCompilationManager::ProcessCompiledComputeKernelShaderM
 					for (int32 ErrorIndex = 0; ErrorIndex < Errors.Num(); ErrorIndex++)
 					{
 						FComputeKernelCompileMessage Message;
-						if (ParseShaderCompilerError(Errors[ErrorIndex], Message))
+						if (ParseShaderCompilerError(Errors[ErrorIndex], CurrentJob.Output.bSucceeded, Message))
 						{
 							ProcessedCompileResults.Messages.AddUnique(Message);
 						}
diff --git a/Engine/Source/Developer/ShaderPreprocessor/Private/stb_preprocess/preprocessor.c b/Engine/Source/Developer/ShaderPreprocessor/Private/stb_preprocess/preprocessor.c
--- a/Engine/Source/Developer/ShaderPreprocessor/Private/stb_preprocess/preprocessor.c
+++ b/Engine/Source/Developer/ShaderPreprocessor/Private/stb_preprocess/preprocessor.c
@@ -4,6 +4,7 @@
 #include "preprocessor.h"
 #include <assert.h>
 #include <ctype.h>
+#include <limits.h>
 #include <stdarg.h>
 #include <stdio.h>
 #include <stdlib.h>
@@ -3963,15 +3964,59 @@ static void process_directive(parse_state* cs, conditional_state* cons)
 		}
 
 		case HASH_line:
-			// copy line directive through to output... existing code will output the newline, so don't do it here
+		{
+			// parse line number and file name and apply these to the parser state. this will trigger a line directive to be emitted with the correct format.
+
+			p = preprocessor_skip_whitespace_simple(p);
+			if (strlen(p) == 0)
+				return;
+
+			char* number_end = NULL;
+			long long int line_number = strtoll(p, &number_end, 10);
+			if (!number_end || number_end == p || llabs(line_number) >= INT_MAX)
 			{
-				size_t out_length = strlen(p);
-				size_t s = arraddnindex(cs->dest, 6);
-				memcpy(cs->dest + s, "#line ", 6);
-				s = arraddnindex(cs->dest, out_length);
-				memcpy(cs->dest + s, p, out_length);
+				do_error_code(cs, PP_RESULT_invalid_line_linenumber, "Invalid #line line number");
+				return;
+			}
+
+			cs->src_line_number = (int)line_number;
+
+			p = preprocessor_skip_whitespace_simple(number_end);
+
+			// file path next. optional - we're done if nothing provided.
+			if (char_is_end_of_line(*p))
+				return;
+
+			// file path should be surrounded by quotes.
+			if (!char_is_quote(*p++))
+			{
+				do_error_code(cs, PP_RESULT_malformed_line_filename, "Malformed #line file name");
+				return;
+			}
+
+			// find end of file path
+			const char* path_end = p;
+			while (*path_end != '"' && !char_is_end_of_line(*path_end))
+				path_end++;
+
+			if (*path_end != '"')
+			{
+				do_error_code(cs, PP_RESULT_malformed_line_filename, "Malformed #line file name");
+				return;
 			}
+
+			ptrdiff_t len = path_end - p;
+			if (len <= 0)
+			{
+				do_error_code(cs, PP_RESULT_malformed_line_filename, "Malformed #line file name");
+				return;
+			}
+
+			// apply new file path
+			cs->filename = stb_arena_alloc_string_length(&c->macro_arena, p, len);
+
 			return;
+		}
 
 		case HASH_pragma:
 		{
diff --git a/Engine/Source/Developer/ShaderPreprocessor/Private/stb_preprocess/preprocessor.h b/Engine/Source/Developer/ShaderPreprocessor/Private/stb_preprocess/preprocessor.h
--- a/Engine/Source/Developer/ShaderPreprocessor/Private/stb_preprocess/preprocessor.h
+++ b/Engine/Source/Developer/ShaderPreprocessor/Private/stb_preprocess/preprocessor.h
@@ -24,6 +24,8 @@ enum
 	PP_RESULT_undef_of_predefined_macro,
 	PP_RESULT_identifier_too_long,
 	PP_RESULT_malformed_pragma_message,
+	PP_RESULT_malformed_line_filename,
+	PP_RESULT_invalid_line_linenumber,
 
 	PP_RESULT_count
 };
