[Compute Framework] Add cvar for triggering RD captures for 1 or more subsequent graph dispatches.

#rnx

diff --git a/Engine/Plugins/Runtime/ComputeFramework/Source/ComputeFramework/Private/ComputeGraphWorker.cpp b/Engine/Plugins/Runtime/ComputeFramework/Source/ComputeFramework/Private/ComputeGraphWorker.cpp
--- a/Engine/Plugins/Runtime/ComputeFramework/Source/ComputeFramework/Private/ComputeGraphWorker.cpp
+++ b/Engine/Plugins/Runtime/ComputeFramework/Source/ComputeFramework/Private/ComputeGraphWorker.cpp
@@ -9,6 +9,7 @@
 #include "ComputeFramework/ComputeGraph.h"
 #include "ComputeFramework/ComputeGraphRenderProxy.h"
 #include "ProfilingDebugging/RealtimeGPUProfiler.h"
+#include "RenderCaptureInterface.h"
 #include "RenderGraphBuilder.h"
 #include "RenderGraphUtils.h"
 #include "RHIDefinitions.h"
@@ -23,6 +24,14 @@ static TAutoConsoleVariable<int32> CVarComputeFrameworkSortSubmit(
 	ECVF_RenderThreadSafe
 );
 
+static int32 GTriggerGPUCaptureDispatches = 0;
+static FAutoConsoleVariableRef CVarComputeFrameworkTriggerGPUCapture(
+	TEXT("r.ComputeFramework.TriggerGPUCaptureDispatches"),
+	GTriggerGPUCaptureDispatches,
+	TEXT("Trigger GPU captures for this many of the subsequent compute graph dispatches."),
+	ECVF_RenderThreadSafe
+);
+
 
 void FComputeGraphTaskWorker::Enqueue(
 	FName InExecutionGroupName, 
@@ -83,6 +92,9 @@ void FComputeGraphTaskWorker::SubmitWork(FRDGBuilder& GraphBuilder, FName InExec
 	RDG_EVENT_SCOPE(GraphBuilder, "ComputeFramework::ExecuteBatches");
 	RDG_GPU_STAT_SCOPE(GraphBuilder, ComputeFramework_ExecuteBatches);
 
+	RenderCaptureInterface::FScopedCapture RenderCapture(GTriggerGPUCaptureDispatches > 0, GraphBuilder, TEXT("FComputeGraphTaskWorker::SubmitWork"));
+	GTriggerGPUCaptureDispatches = FMath::Max(GTriggerGPUCaptureDispatches - 1, 0);
+
 	// Currently poll readbacks once at end of frame.
 	if (InExecutionGroupName == ComputeTaskExecutionGroup::EndOfFrameUpdate)
 	{
