[Compute Framework] Add post graph dispatch virtual function that can be used to perform post execution actions.

#rnx

diff --git a/Engine/Plugins/Runtime/ComputeFramework/Source/ComputeFramework/Private/ComputeGraphWorker.cpp b/Engine/Plugins/Runtime/ComputeFramework/Source/ComputeFramework/Private/ComputeGraphWorker.cpp
--- a/Engine/Plugins/Runtime/ComputeFramework/Source/ComputeFramework/Private/ComputeGraphWorker.cpp
+++ b/Engine/Plugins/Runtime/ComputeFramework/Source/ComputeFramework/Private/ComputeGraphWorker.cpp
@@ -326,6 +326,21 @@ void FComputeGraphTaskWorker::SubmitWork(FRDGBuilder& GraphBuilder, FName InExec
 		}
 	}
 
+	// Let proxies perform any post-graph-dispatch actions.
+	// TODO: Given few proxies currently implement this method, it may be better to record which proxies require this callback rather than calling on all.
+	for (FSubmitDescription const& SubmitDesc : SubmitDescs)
+	{
+		const int32 GraphIndex = SubmitDesc.GraphIndex;
+		FGraphInvocation const& GraphInvocation = GraphInvocations[GraphIndex];
+		for (const FComputeDataProviderRenderProxy* Proxy : GraphInvocation.DataProviderRenderProxies)
+		{
+			if (Proxy)
+			{
+				Proxy->PostGraphDispatch(GraphBuilder);
+			}
+		}
+	}
+
 	// Release any graph resources at the end of graph execution.
 	GraphBuilder.AddPass(
 		{},
diff --git a/Engine/Plugins/Runtime/ComputeFramework/Source/ComputeFramework/Public/ComputeFramework/ComputeDataProvider.h b/Engine/Plugins/Runtime/ComputeFramework/Source/ComputeFramework/Public/ComputeFramework/ComputeDataProvider.h
--- a/Engine/Plugins/Runtime/ComputeFramework/Source/ComputeFramework/Public/ComputeFramework/ComputeDataProvider.h
+++ b/Engine/Plugins/Runtime/ComputeFramework/Source/ComputeFramework/Public/ComputeFramework/ComputeDataProvider.h
@@ -124,6 +124,9 @@ public:
 	/** Data for any readbacks that should be performed. */
 	virtual void GetReadbackData(TArray<FReadbackData>& OutReadbackData) const {}
 
+	/** Called after graph dispatched. */
+	virtual void PostGraphDispatch(FRDGBuilder& GraphBuilder) const {}
+
 protected:
 	/** Helper for making an FStridedView over the FDispatchData. */
 	template <typename ElementType>
