[ComputeFramework] Always initialize KernelResources in BeginCacheForCookedPlatformData
* Prevents crash when cooking without calling UpdateResources()

diff --git a/Engine/Plugins/Runtime/ComputeFramework/Source/ComputeFramework/Private/ComputeGraph.cpp b/Engine/Plugins/Runtime/ComputeFramework/Source/ComputeFramework/Private/ComputeGraph.cpp
--- a/Engine/Plugins/Runtime/ComputeFramework/Source/ComputeFramework/Private/ComputeGraph.cpp
+++ b/Engine/Plugins/Runtime/ComputeFramework/Source/ComputeFramework/Private/ComputeGraph.cpp
@@ -788,6 +788,8 @@ void UComputeGraph::BeginCacheForCookedPlatformData(ITargetPlatform const* Targe
 	TArray<FName> ShaderFormats;
 	TargetPlatform->GetAllTargetedShaderFormats(ShaderFormats);
 
+	KernelResources.SetNum(KernelInvocations.Num());
+
 	for (int32 KernelIndex = 0; KernelIndex < KernelInvocations.Num(); ++KernelIndex)
 	{
 		TArray< TUniquePtr<FComputeKernelResource> >& Resources = KernelResources[KernelIndex].CachedKernelResourcesForCooking.FindOrAdd(TargetPlatform);
