[ComputeFramework] Improve logic around updating readbacks to fix crash when there are active readbacks but no work to submit.

#rnx

diff --git a/Engine/Plugins/Runtime/ComputeFramework/Source/ComputeFramework/Private/ComputeGraphWorker.cpp b/Engine/Plugins/Runtime/ComputeFramework/Source/ComputeFramework/Private/ComputeGraphWorker.cpp
--- a/Engine/Plugins/Runtime/ComputeFramework/Source/ComputeFramework/Private/ComputeGraphWorker.cpp
+++ b/Engine/Plugins/Runtime/ComputeFramework/Source/ComputeFramework/Private/ComputeGraphWorker.cpp
@@ -77,7 +77,8 @@ void FComputeGraphTaskWorker::Abort(const UObject* InOwnerPointer)
 
 bool FComputeGraphTaskWorker::HasWork(FName InExecutionGroupName) const
 {
-	if (!ActiveAsyncReadbacks.IsEmpty())
+	// Currently poll readbacks once at end of frame.
+	if (!ActiveAsyncReadbacks.IsEmpty() && InExecutionGroupName == ComputeTaskExecutionGroup::EndOfFrameUpdate)
 	{
 		return true;
 	}
@@ -101,11 +102,17 @@ void FComputeGraphTaskWorker::SubmitWork(FRDGBuilder& GraphBuilder, FName InExec
 		UpdateReadbacks();
 	}
 
+	TArray<FGraphInvocation> const* FoundGraphInvocations = GraphInvocationsPerGroup.Find(InExecutionGroupName);
+	if (!FoundGraphInvocations || FoundGraphInvocations->Num() == 0)
+	{
+		return;
+	}
+
 	// Reset our scratch memory arrays.
 	SubmitDescs.Reset();
 	Shaders.Reset();
 
-	TArray<FGraphInvocation> const& GraphInvocations = GraphInvocationsPerGroup.FindChecked(InExecutionGroupName);
+	TArray<FGraphInvocation> const& GraphInvocations = *FoundGraphInvocations;
 	for (int32 GraphIndex = 0; GraphIndex < GraphInvocations.Num(); ++GraphIndex)
 	{
 		FGraphInvocation const& GraphInvocation = GraphInvocations[GraphIndex];
