Explicitly declare cook dependencies on shader files for UComputeGraph.
Reenable incremental cooking for UComputeGraph.
#tests Test project cooked with -IterativeValidate

diff --git a/Engine/Config/BaseEditor.ini b/Engine/Config/BaseEditor.ini
--- a/Engine/Config/BaseEditor.ini
+++ b/Engine/Config/BaseEditor.ini
@@ -495,7 +495,6 @@ SaveUnversioned=true
 
 +IterativeClassScriptPackageAllowList=Allow,<EngineRoot>
 
-+IterativeClassDenyList=/Script/ComputeFramework.ComputeGraph				; UE-203846
 +IterativeClassDenyList=/Script/CustomizableObject.CustomizableObject		; UE-224744
 +IterativeClassDenyList=/Script/Engine.BlueprintGeneratedClass 				; Not yet documented
 +IterativeClassDenyList=/Script/Engine.World 								; Not yet documented
diff --git a/Engine/Plugins/Runtime/ComputeFramework/Source/ComputeFramework/Private/ComputeGraph.cpp b/Engine/Plugins/Runtime/ComputeFramework/Source/ComputeFramework/Private/ComputeGraph.cpp
--- a/Engine/Plugins/Runtime/ComputeFramework/Source/ComputeFramework/Private/ComputeGraph.cpp
+++ b/Engine/Plugins/Runtime/ComputeFramework/Source/ComputeFramework/Private/ComputeGraph.cpp
@@ -14,12 +14,15 @@
 #include "ComputeFramework/ComputeKernelSource.h"
 #include "ComputeFramework/ComputeSource.h"
 #include "ComputeFramework/ShaderParameterMetadataAllocation.h"
+#include "DataDrivenShaderPlatformInfo.h"
 #include "GameFramework/Actor.h"
 #include "Interfaces/ITargetPlatform.h"
 #include "Misc/App.h"
+#include "Serialization/CompactBinarySerialization.h"
+#include "Serialization/CompactBinaryWriter.h"
 #include "ShaderParameterMetadataBuilder.h"
+#include "UObject/ObjectSaveContext.h"
 #include "UObject/Package.h"
-#include "DataDrivenShaderPlatformInfo.h"
 
 #include UE_INLINE_GENERATED_CPP_BY_NAME(ComputeGraph)
 
@@ -95,6 +98,97 @@ void UComputeGraph::PostLoad()
 	}
 }
 
+#if WITH_EDITOR
+
+/** Shader dependency serialization code used for incremental cooking. */
+namespace ComputeFramework::Private
+{ 
+	constexpr int32 HashDependenciesForCookArgsVersion = 1;
+	
+	void HashDependenciesForCook(FCbFieldViewIterator Args, UE::Cook::FCookDependencyContext& Context)
+	{
+		FCbFieldViewIterator ArgField(Args);
+		int32 ArgsVersion = -1;
+		ArgsVersion = (ArgField++).AsInt32();
+
+		bool bValid = ArgsVersion == HashDependenciesForCookArgsVersion;
+		
+		TArray<FName> ShaderFormats;
+		if (bValid)
+		{
+			bValid = LoadFromCompactBinary(ArgField++, ShaderFormats);
+		}
+		
+		TArray<FString> ShaderVirtualPaths;
+		if (bValid)
+		{
+			bValid = LoadFromCompactBinary(ArgField++, ShaderVirtualPaths);
+		}
+
+		if (!bValid)
+		{
+			Context.LogInvalidated(FString::Printf(TEXT("Failed to serialize UComputeGraph cook dependencies. ArgsVersion %d."), ArgsVersion));
+			return;
+		}
+
+		for (FName ShaderFormat : ShaderFormats)
+		{
+			const EShaderPlatform ShaderPlatform = ShaderFormatToLegacyShaderPlatform(ShaderFormat);
+
+			for (FString const& ShaderVirtualPath : ShaderVirtualPaths)
+			{
+				const FSHAHash Hash = GetShaderFileHash(*ShaderVirtualPath, ShaderPlatform);
+				Context.Update(&Hash.Hash, sizeof(Hash.Hash));
+			}
+		}
+	}
+}
+
+UE_COOK_DEPENDENCY_FUNCTION(HashComputeGraphDependenciesForCook, ComputeFramework::Private::HashDependenciesForCook);
+
+#endif // WITH_EDITOR
+
+
+void UComputeGraph::PreSave(FObjectPreSaveContext ObjectSaveContext)
+{
+	Super::PreSave(ObjectSaveContext);
+#if WITH_EDITOR
+	if (ObjectSaveContext.IsCooking())
+	{
+		TArray<FName> ShaderFormats;
+		if (ITargetPlatform const* TargetPlatform = ObjectSaveContext.GetTargetPlatform())
+		{
+			TargetPlatform->GetAllTargetedShaderFormats(ShaderFormats);
+		}
+		
+		TArray<FString> ShaderVirtualPaths;
+		ShaderVirtualPaths.Reserve(DataInterfaces.Num() + 1);
+		
+		// Add dependency on compute graph entry shader file (and its includes).
+		ShaderVirtualPaths.Add(TEXT("/Plugin/ComputeFramework/Private/ComputeKernel.usf"));
+		
+		// Add dependency on data interface shader template files (and their includes).
+		for (UComputeDataInterface const* DataInterface : DataInterfaces)
+		{
+			if (DataInterface != nullptr)
+			{
+				if (TCHAR const* ShaderVirtualPath = DataInterface->GetShaderVirtualPath())
+				{
+					ShaderVirtualPaths.Add(ShaderVirtualPath);
+				}
+			}
+		}
+
+		FCbWriter Writer;
+		Writer << ComputeFramework::Private::HashDependenciesForCookArgsVersion;
+		Writer << ShaderFormats;
+		Writer << ShaderVirtualPaths;
+
+		ObjectSaveContext.AddCookBuildDependency(UE::Cook::FCookDependency::Function(UE_COOK_DEPENDENCY_FUNCTION_CALL(HashComputeGraphDependenciesForCook), Writer.Save()));
+	}
+#endif
+}
+
 bool UComputeGraph::ValidateGraph(FString* OutErrors)
 {
 	// todo[CF]:
diff --git a/Engine/Plugins/Runtime/ComputeFramework/Source/ComputeFramework/Public/ComputeFramework/ComputeGraph.h b/Engine/Plugins/Runtime/ComputeFramework/Source/ComputeFramework/Public/ComputeFramework/ComputeGraph.h
--- a/Engine/Plugins/Runtime/ComputeFramework/Source/ComputeFramework/Public/ComputeFramework/ComputeGraph.h
+++ b/Engine/Plugins/Runtime/ComputeFramework/Source/ComputeFramework/Public/ComputeFramework/ComputeGraph.h
@@ -138,6 +138,7 @@ protected:
 	//~ Begin UObject Interface.
 	void Serialize(FArchive& Ar) override;
 	void PostLoad() override;
+	void PreSave(FObjectPreSaveContext ObjectSaveContext);
 #if WITH_EDITOR
 	void BeginCacheForCookedPlatformData(ITargetPlatform const* TargetPlatform) override;
 	bool IsCachedCookedPlatformDataLoaded(ITargetPlatform const* TargetPlatform) override;
