[ComputeFramework] Add optional PreSubmit and PostSubmit callbacks to data interfaces to perform additional work.

Replaces PostGraphDispatch().

#rnx

diff --git a/Engine/Plugins/PCG/Source/PCG/Private/Compute/DataInterfaces/PCGInstanceDataInterface.cpp b/Engine/Plugins/PCG/Source/PCG/Private/Compute/DataInterfaces/PCGInstanceDataInterface.cpp
--- a/Engine/Plugins/PCG/Source/PCG/Private/Compute/DataInterfaces/PCGInstanceDataInterface.cpp
+++ b/Engine/Plugins/PCG/Source/PCG/Private/Compute/DataInterfaces/PCGInstanceDataInterface.cpp
@@ -230,16 +230,16 @@ void FPCGInstanceDataProviderProxy::GatherDispatchData(FDispatchData const& InDi
 	}
 }
 
-void FPCGInstanceDataProviderProxy::PostGraphDispatch(FRDGBuilder& GraphBuilder) const
+void FPCGInstanceDataProviderProxy::PostSubmit(FRDGBuilder& InGraphBuilder) const
 {
-	TRACE_CPUPROFILER_EVENT_SCOPE(FPCGInstanceDataProviderProxy::PostGraphDispatch);
+	TRACE_CPUPROFILER_EVENT_SCOPE(FPCGInstanceDataProviderProxy::PostSubmit);
 
 	if (NumInstancesAllPrimitives == 0)
 	{
 		return;
 	}
 
-	const FSceneComputeUpdatesBlackboardEntry* SceneComputeUpdatesBB = GraphBuilder.Blackboard.Get<FSceneComputeUpdatesBlackboardEntry>();
+	const FSceneComputeUpdatesBlackboardEntry* SceneComputeUpdatesBB = InGraphBuilder.Blackboard.Get<FSceneComputeUpdatesBlackboardEntry>();
 	if (!SceneComputeUpdatesBB || !ensure(SceneComputeUpdatesBB->SceneComputeUpdates))
 	{
 		UE_LOG(LogPCG, Warning, TEXT("Scene Compute Updates interface missing from graph builder blackboard. Scene will not be updated."));
diff --git a/Engine/Plugins/PCG/Source/PCG/Private/Compute/DataInterfaces/PCGInstanceDataInterface.h b/Engine/Plugins/PCG/Source/PCG/Private/Compute/DataInterfaces/PCGInstanceDataInterface.h
--- a/Engine/Plugins/PCG/Source/PCG/Private/Compute/DataInterfaces/PCGInstanceDataInterface.h
+++ b/Engine/Plugins/PCG/Source/PCG/Private/Compute/DataInterfaces/PCGInstanceDataInterface.h
@@ -27,6 +27,7 @@ public:
 	TCHAR const* GetShaderVirtualPath() const override;
 	void GetShaderHash(FString& InOutKey) const override;
 	void GetHLSL(FString& OutHLSL, FString const& InDataInterfaceName) const override;
+	bool GetRequiresPostSubmitCall() const override { return true; }
 	UComputeDataProvider* CreateDataProvider(TObjectPtr<UObject> InBinding, uint64 InInputMask, uint64 InOutputMask) const override;
 	//~ End UComputeDataInterface Interface
 
@@ -74,7 +75,7 @@ public:
 	bool IsValid(FValidationData const& InValidationData) const override;
 	void AllocateResources(FRDGBuilder& GraphBuilder, FAllocationData const& InAllocationData);
 	void GatherDispatchData(FDispatchData const& InDispatchData) override;
-	void PostGraphDispatch(FRDGBuilder& GraphBuilder) const override;
+	void PostSubmit(FRDGBuilder& InGraphBuilder) const override;
 	//~ End FComputeDataProviderRenderProxy Interface
 
 protected:
diff --git a/Engine/Plugins/Runtime/ComputeFramework/Source/ComputeFramework/Private/ComputeFramework/ComputeGraphRenderProxy.h b/Engine/Plugins/Runtime/ComputeFramework/Source/ComputeFramework/Private/ComputeFramework/ComputeGraphRenderProxy.h
--- a/Engine/Plugins/Runtime/ComputeFramework/Source/ComputeFramework/Private/ComputeFramework/ComputeGraphRenderProxy.h
+++ b/Engine/Plugins/Runtime/ComputeFramework/Source/ComputeFramework/Private/ComputeFramework/ComputeGraphRenderProxy.h
@@ -32,6 +32,10 @@ public:
 		TArray<int32> BoundProviderIndices;
 		/** Array of indices for data providers that should trigger a readback. */
 		TArray<int32> ReadbackProviderIndices;
+		/** Indices of data providers that require PreSubmit call. */
+		TArray<int32> PreSubmitProviderIndices;
+		/** Indices of data providers that require PostSubmit call. */
+		TArray<int32> PostSubmitProviderIndices;
 		/** Same size as BoundProviderIndices, non-primary data providers will be forced to present a full view to its data (unified)*/
 		TArray<bool> BoundProviderIsPrimary;
 		/** The index of the special execution data provider in the full graph data provider array. */
diff --git a/Engine/Plugins/Runtime/ComputeFramework/Source/ComputeFramework/Private/ComputeGraph.cpp b/Engine/Plugins/Runtime/ComputeFramework/Source/ComputeFramework/Private/ComputeGraph.cpp
--- a/Engine/Plugins/Runtime/ComputeFramework/Source/ComputeFramework/Private/ComputeGraph.cpp
+++ b/Engine/Plugins/Runtime/ComputeFramework/Source/ComputeFramework/Private/ComputeGraph.cpp
@@ -477,9 +477,22 @@ FComputeGraphRenderProxy* UComputeGraph::CreateRenderProxy() const
 							}
 
 							// If the data interface is requesting readback and is a kernel output, register it.
-							if (!GraphEdge.bKernelInput && DataInterface->GetRequiresReadback())
+							if (!GraphEdge.bKernelInput)
 							{
-								Invocation.ReadbackProviderIndices.AddUnique(GraphEdge.DataInterfaceIndex);
+								if (DataInterface->GetRequiresReadback())
+								{
+									Invocation.ReadbackProviderIndices.AddUnique(GraphEdge.DataInterfaceIndex);
+								}
+
+								if (DataInterface->GetRequiresPreSubmitCall())
+								{
+									Invocation.PreSubmitProviderIndices.AddUnique(GraphEdge.DataInterfaceIndex);
+								}
+
+								if (DataInterface->GetRequiresPostSubmitCall())
+								{
+									Invocation.PostSubmitProviderIndices.AddUnique(GraphEdge.DataInterfaceIndex);
+								}
 							}
 						}
 					}
diff --git a/Engine/Plugins/Runtime/ComputeFramework/Source/ComputeFramework/Private/ComputeGraphWorker.cpp b/Engine/Plugins/Runtime/ComputeFramework/Source/ComputeFramework/Private/ComputeGraphWorker.cpp
--- a/Engine/Plugins/Runtime/ComputeFramework/Source/ComputeFramework/Private/ComputeGraphWorker.cpp
+++ b/Engine/Plugins/Runtime/ComputeFramework/Source/ComputeFramework/Private/ComputeGraphWorker.cpp
@@ -236,6 +236,15 @@ void FComputeGraphTaskWorker::SubmitWork(FRDGBuilder& GraphBuilder, FName InExec
 
 		RDG_EVENT_SCOPE(GraphBuilder, "%s:%s:%s", *GraphInvocation.OwnerName.ToString(), *GraphRenderProxy->GraphName.ToString(), *KernelInvocation.KernelName);
 
+		// Pre submit calls.
+		for (int32 DataProviderIndex : KernelInvocation.PreSubmitProviderIndices)
+		{
+			if (FComputeDataProviderRenderProxy* DataProviderProxy = GraphInvocation.DataProviderRenderProxies[DataProviderIndex])
+			{
+				DataProviderProxy->PreSubmit(GraphBuilder);
+			}
+		}
+
 		//todo[CF]: GetDispatchThreadCount() should take the bIsUnified flag directly.
 		ThreadCounts.Reset();
 		int32 NumSubInvocations = GraphInvocation.DataProviderRenderProxies[KernelInvocation.ExecutionProviderIndex]->GetDispatchThreadCount(ThreadCounts);
@@ -303,6 +312,15 @@ void FComputeGraphTaskWorker::SubmitWork(FRDGBuilder& GraphBuilder, FName InExec
 			);
 		}
 
+		// Post submit calls.
+		for (int32 DataProviderIndex : KernelInvocation.PostSubmitProviderIndices)
+		{
+			if (FComputeDataProviderRenderProxy* DataProviderProxy = GraphInvocation.DataProviderRenderProxies[DataProviderIndex])
+			{
+				DataProviderProxy->PostSubmit(GraphBuilder);
+			}
+		}
+
 		// Enqueue readbacks.
 		for (int32 DataProviderIndex : KernelInvocation.ReadbackProviderIndices)
 		{
@@ -333,19 +351,6 @@ void FComputeGraphTaskWorker::SubmitWork(FRDGBuilder& GraphBuilder, FName InExec
 		}
 	}
 
-	// Let proxies perform any post-graph-dispatch actions.
-	// TODO: Given few proxies currently implement this method, it may be better to record which proxies require this callback rather than calling on all.
-	for (FGraphInvocation const& GraphInvocation : GraphInvocations)
-	{
-		for (const FComputeDataProviderRenderProxy* Proxy : GraphInvocation.DataProviderRenderProxies)
-		{
-			if (Proxy)
-			{
-				Proxy->PostGraphDispatch(GraphBuilder);
-			}
-		}
-	}
-
 	// Release any graph resources at the end of graph execution.
 	GraphBuilder.AddPostExecuteCallback(
 		[this, InExecutionGroupName]
diff --git a/Engine/Plugins/Runtime/ComputeFramework/Source/ComputeFramework/Public/ComputeFramework/ComputeDataInterface.h b/Engine/Plugins/Runtime/ComputeFramework/Source/ComputeFramework/Public/ComputeFramework/ComputeDataInterface.h
--- a/Engine/Plugins/Runtime/ComputeFramework/Source/ComputeFramework/Public/ComputeFramework/ComputeDataInterface.h
+++ b/Engine/Plugins/Runtime/ComputeFramework/Source/ComputeFramework/Public/ComputeFramework/ComputeDataInterface.h
@@ -49,6 +49,10 @@ public:
 	virtual void GetHLSL(FString& OutHLSL, FString const& InDataInterfaceName) const {}
 	/** Whether a GPU->CPU readback should be scheduled after this data is used as a kernel output. */
 	virtual bool GetRequiresReadback() const { return false; }
+	/** Whether PreSubmit should be called when this data is used as a kernel output. */
+	virtual bool GetRequiresPreSubmitCall() const { return false; }
+	/** Whether PostSubmit should be called when this data is used as a kernel output. */
+	virtual bool GetRequiresPostSubmitCall() const { return false; }
 	/** Get class of UObject required to instantiate a UComputeDataProvider from this interface. */
 	virtual UClass* GetBindingType() const { return nullptr; }
 	/** Instantiate an associated UComputeDataProvider. */
diff --git a/Engine/Plugins/Runtime/ComputeFramework/Source/ComputeFramework/Public/ComputeFramework/ComputeDataProvider.h b/Engine/Plugins/Runtime/ComputeFramework/Source/ComputeFramework/Public/ComputeFramework/ComputeDataProvider.h
--- a/Engine/Plugins/Runtime/ComputeFramework/Source/ComputeFramework/Public/ComputeFramework/ComputeDataProvider.h
+++ b/Engine/Plugins/Runtime/ComputeFramework/Source/ComputeFramework/Public/ComputeFramework/ComputeDataProvider.h
@@ -124,8 +124,11 @@ public:
 	/** Data for any readbacks that should be performed. */
 	virtual void GetReadbackData(TArray<FReadbackData>& OutReadbackData) const {}
 
-	/** Called after graph dispatched. */
-	virtual void PostGraphDispatch(FRDGBuilder& GraphBuilder) const {}
+	/** Called immediately prior to every kernel dispatch that outputs to this data. */
+	virtual void PreSubmit(FRDGBuilder& InGraphBuilder) const {}
+
+	/** Called immediately after every kernel dispatch that outputs to this data. */
+	virtual void PostSubmit(FRDGBuilder& InGraphBuilder) const {}
 
 protected:
 	/** Helper for making an FStridedView over the FDispatchData. */
@@ -162,4 +165,8 @@ protected:
 	UE_DEPRECATED(5.2, "Convert to using the new GatherDispatchData() that takes FDispatchData.")
 	virtual void GatherDispatchData(FDispatchSetup const& InDispatchSetup, FCollectedDispatchData& InOutDispatchData) {}
 	//~ End deprecation of old GatherDispatchData() interface.
+
+public:
+	UE_DEPRECATED(5.6, "Use PostSubmit instead.")
+	virtual void PostGraphDispatch(FRDGBuilder& GraphBuilder) const {}
 };
