Fixed serialization of default FShaderValueTypeHandle

The function was unable to roundtrip a default handle but hasn't needed to in practice because delta serialization has skipped serialization of the default until now.

diff --git a/Engine/Plugins/Runtime/ComputeFramework/Source/ComputeFramework/Private/ShaderParamTypeDefinition.cpp b/Engine/Plugins/Runtime/ComputeFramework/Source/ComputeFramework/Private/ShaderParamTypeDefinition.cpp
--- a/Engine/Plugins/Runtime/ComputeFramework/Source/ComputeFramework/Private/ShaderParamTypeDefinition.cpp
+++ b/Engine/Plugins/Runtime/ComputeFramework/Source/ComputeFramework/Private/ShaderParamTypeDefinition.cpp
@@ -632,46 +632,43 @@ FArchive& operator<<(FArchive& InArchive, FShaderValueTypeHandle& InHandle)
 	{
 		ValueTypePtr = &ValueTypeTemp;
 	}
+	else if (InHandle.ValueTypePtr)
+	{
+		ValueTypePtr = const_cast<FShaderValueType*>(InHandle.ValueTypePtr);
+	}
 	else
 	{
-		ValueTypePtr = const_cast<FShaderValueType *>(InHandle.ValueTypePtr);
+		ValueTypeTemp.Type = EShaderFundamentalType::None;
+		ValueTypePtr = &ValueTypeTemp;
 	}
 
-	if (ValueTypePtr)
-	{
-		InArchive << ValueTypePtr->Type;
+	InArchive << ValueTypePtr->Type;
 
-		// When the object version is added, bIsDynamicArray is also added
-		if (InArchive.CustomVer(FComputeFrameworkObjectVersion::GUID) >= FComputeFrameworkObjectVersion::InitialVersion)
-		{
-			InArchive << ValueTypePtr->bIsDynamicArray;
-		}
+	// When the object version is added, bIsDynamicArray is also added
+	if (InArchive.CustomVer(FComputeFrameworkObjectVersion::GUID) >= FComputeFrameworkObjectVersion::InitialVersion)
+	{
+		InArchive << ValueTypePtr->bIsDynamicArray;
+	}
 		
-		if (ValueTypePtr->Type == EShaderFundamentalType::Struct)
+	if (ValueTypePtr->Type == EShaderFundamentalType::Struct)
+	{
+		InArchive << ValueTypePtr->Name;
+		InArchive << ValueTypePtr->StructElements;
+	}
+	else
+	{
+		InArchive << ValueTypePtr->DimensionType;
+
+		if (ValueTypePtr->DimensionType == EShaderFundamentalDimensionType::Vector)
 		{
-			InArchive << ValueTypePtr->Name;
-			InArchive << ValueTypePtr->StructElements;
+			InArchive << ValueTypePtr->VectorElemCount;
 		}
-		else
+		else if (ValueTypePtr->DimensionType == EShaderFundamentalDimensionType::Matrix)
 		{
-			InArchive << ValueTypePtr->DimensionType;
-
-			if (ValueTypePtr->DimensionType == EShaderFundamentalDimensionType::Vector)
-			{
-				InArchive << ValueTypePtr->VectorElemCount;
-			}
-			else if (ValueTypePtr->DimensionType == EShaderFundamentalDimensionType::Matrix)
-			{
-				InArchive << ValueTypePtr->MatrixRowCount;
-				InArchive << ValueTypePtr->MatrixColumnCount;
-			}
+			InArchive << ValueTypePtr->MatrixRowCount;
+			InArchive << ValueTypePtr->MatrixColumnCount;
 		}
 	}
-	else if (InArchive.IsSaving())
-	{
-		EShaderFundamentalType Type = EShaderFundamentalType::None;
-		InArchive << Type; 
-	}
 
 	if (InArchive.IsLoading())
 	{
