[Compute Framework] FComputeGraphTaskWorker::Abort no longer clears active readbacks, which gives readback requester a chance to to decide if it still want to use the readback data or discard it.

diff --git a/Engine/Plugins/Runtime/ComputeFramework/Source/ComputeFramework/Private/ComputeGraphWorker.cpp b/Engine/Plugins/Runtime/ComputeFramework/Source/ComputeFramework/Private/ComputeGraphWorker.cpp
--- a/Engine/Plugins/Runtime/ComputeFramework/Source/ComputeFramework/Private/ComputeGraphWorker.cpp
+++ b/Engine/Plugins/Runtime/ComputeFramework/Source/ComputeFramework/Private/ComputeGraphWorker.cpp
@@ -65,14 +65,22 @@ void FComputeGraphTaskWorker::Abort(const UObject* InOwnerPointer)
 			}
 		}
 	}
-
-	for (int I = ActiveAsyncReadbacks.Num() - 1; I >= 0; --I)
-	{
-		if (ActiveAsyncReadbacks[I].OwnerPointer == InOwnerPointer)
-		{
-			ActiveAsyncReadbacks.RemoveAtSwap(I);
-		}
-	}
+	
+	// Not clearing ActiveAsyncReadbacks here because the system requesting 
+	// the readback data isn't always InOwnerPointer, instead it can be some external system that
+	// simply wants to readback the latest data produced by InOwnerPointer.
+	
+	// When InOwnerPointer becomes invalid at a frame, any readback requested in the previous frame 
+	// can still be useful to the requester, but it is likely not ready until 2 or 3 frames later.
+	// So the worker needs to keep polling for that last bit of data for a few extra frames.
+	
+	// The Requester system can use OnDataAvailable callback to decide if it actually wants
+	// to use the data or discard it.
+	
+	// For example, when baking level sequence to geometry cache, the baker can request a readback right before the frame 
+	// where the skeletal mesh despawns. The data for that frame won't be ready until 2 or 3 frames after 
+	// skeletal mesh has been destroyed, so we need the worker to keep polling for that data for 2-3 extra frames.
+	
 }
 
 bool FComputeGraphTaskWorker::HasWork(FName InExecutionGroupName) const
