ComputeFramework: Add LayoutParams to the shader DDC key. Previously it was possible to get a key collision in the case of different layouts.

diff --git a/Engine/Plugins/Runtime/ComputeFramework/Source/ComputeFramework/Private/ComputeKernelShaderMap.cpp b/Engine/Plugins/Runtime/ComputeFramework/Source/ComputeFramework/Private/ComputeKernelShaderMap.cpp
--- a/Engine/Plugins/Runtime/ComputeFramework/Source/ComputeFramework/Private/ComputeKernelShaderMap.cpp
+++ b/Engine/Plugins/Runtime/ComputeFramework/Source/ComputeFramework/Private/ComputeKernelShaderMap.cpp
@@ -119,6 +119,9 @@ void FComputeKernelShaderMapId::AppendKeyString(FString& OutKeyString) const
 		OutKeyString.AppendChar('_');
 	}
 
+	// Include the LayoutParams to differentiate between the frozen memory layout differences between platforms
+	LayoutParams.AppendKeyString(OutKeyString);
+
 	// Add the inputs for any shaders that are stored inline in the shader map
 	AppendKeyStringShaderDependencies(MakeArrayView(ShaderTypeDependencies), LayoutParams, OutKeyString);
 }
