[ComputeFramework] Add lock to GloballyKnownValueTypes set
* Causing crash on compile from PCG GPU graph compiler

#rnx

diff --git a/Engine/Plugins/Runtime/ComputeFramework/Source/ComputeFramework/Private/ShaderParamTypeDefinition.cpp b/Engine/Plugins/Runtime/ComputeFramework/Source/ComputeFramework/Private/ShaderParamTypeDefinition.cpp
--- a/Engine/Plugins/Runtime/ComputeFramework/Source/ComputeFramework/Private/ShaderParamTypeDefinition.cpp
+++ b/Engine/Plugins/Runtime/ComputeFramework/Source/ComputeFramework/Private/ShaderParamTypeDefinition.cpp
@@ -39,6 +39,7 @@ struct HandleKeyFuncs : BaseKeyFuncs<FShaderValueTypeHandle,FShaderValueTypeHand
 };
 
 static TSet<FShaderValueTypeHandle, HandleKeyFuncs> GloballyKnownValueTypes;
+static FRWLock GloballyKnownValueTypesLock;
 
 
 bool FShaderValueTypeHandle::Serialize(FArchive& Ar)
@@ -243,12 +244,18 @@ FShaderValueTypeHandle FShaderValueType::GetOrCreate(FShaderValueType&& InValueT
 	FShaderValueTypeHandle Handle;
 	Handle.ValueTypePtr = &InValueType;
 
-	FShaderValueTypeHandle *FoundHandle = GloballyKnownValueTypes.Find(Handle);
-	if (FoundHandle)
 	{
-		return *FoundHandle;
+		FReadScopeLock Lock(GloballyKnownValueTypesLock);
+	
+		FShaderValueTypeHandle* FoundHandle = GloballyKnownValueTypes.Find(Handle);
+
+		if (FoundHandle)
+		{
+			return *FoundHandle;
+		}
 	}
 
+	FWriteScopeLock Lock(GloballyKnownValueTypesLock);
 	Handle.ValueTypePtr = new FShaderValueType(MoveTemp(InValueType));
 	GloballyKnownValueTypes.Add(Handle);
 	return Handle;
