[ComputeFramework] Fix remaining race condition, avoid two threads inserting a new element simultaneously.

#rnx

diff --git a/Engine/Plugins/Runtime/ComputeFramework/Source/ComputeFramework/Private/ShaderParamTypeDefinition.cpp b/Engine/Plugins/Runtime/ComputeFramework/Source/ComputeFramework/Private/ShaderParamTypeDefinition.cpp
--- a/Engine/Plugins/Runtime/ComputeFramework/Source/ComputeFramework/Private/ShaderParamTypeDefinition.cpp
+++ b/Engine/Plugins/Runtime/ComputeFramework/Source/ComputeFramework/Private/ShaderParamTypeDefinition.cpp
@@ -247,15 +247,19 @@ FShaderValueTypeHandle FShaderValueType::GetOrCreate(FShaderValueType&& InValueT
 	{
 		FReadScopeLock Lock(GloballyKnownValueTypesLock);
 	
-		FShaderValueTypeHandle* FoundHandle = GloballyKnownValueTypes.Find(Handle);
-
-		if (FoundHandle)
+		if (FShaderValueTypeHandle* FoundHandle = GloballyKnownValueTypes.Find(Handle))
 		{
 			return *FoundHandle;
 		}
 	}
 
 	FWriteScopeLock Lock(GloballyKnownValueTypesLock);
+
+	if (FShaderValueTypeHandle* FoundHandle = GloballyKnownValueTypes.Find(Handle))
+	{
+		return *FoundHandle;
+	}
+
 	Handle.ValueTypePtr = new FShaderValueType(MoveTemp(InValueType));
 	GloballyKnownValueTypes.Add(Handle);
 	return Handle;
