[ComputeFramework] Add create/initialize/reset paths for data providers to enable pooling of compute graph instances, avoiding expensive creation of data providers.

#rnx

diff --git a/Engine/Plugins/PCG/Source/PCG/Private/Compute/DataInterfaces/PCGLandscapeDataInterface.cpp b/Engine/Plugins/PCG/Source/PCG/Private/Compute/DataInterfaces/PCGLandscapeDataInterface.cpp
--- a/Engine/Plugins/PCG/Source/PCG/Private/Compute/DataInterfaces/PCGLandscapeDataInterface.cpp
+++ b/Engine/Plugins/PCG/Source/PCG/Private/Compute/DataInterfaces/PCGLandscapeDataInterface.cpp
@@ -290,7 +290,7 @@ UComputeDataProvider* UPCGLandscapeDataInterface::CreateDataProvider(TObjectPtr<
 					{
 						if (TestLandscape(Landscape))
 						{
-							DataProvider->Initialize(Landscape, ComponentBounds);
+							DataProvider->Init(Landscape, ComponentBounds);
 							break;
 						}
 					}
@@ -302,7 +302,7 @@ UComputeDataProvider* UPCGLandscapeDataInterface::CreateDataProvider(TObjectPtr<
 	return DataProvider;
 }
 
-void UPCGLandscapeDataProvider::Initialize(ALandscape* InLandscape, const FBox& Bounds)
+void UPCGLandscapeDataProvider::Init(ALandscape* InLandscape, const FBox& Bounds)
 {
 	TRACE_CPUPROFILER_EVENT_SCOPE(UPCGLandscapeDataProvider::Initialize);
 	check(IsInGameThread());
diff --git a/Engine/Plugins/PCG/Source/PCG/Private/Compute/DataInterfaces/PCGLandscapeDataInterface.h b/Engine/Plugins/PCG/Source/PCG/Private/Compute/DataInterfaces/PCGLandscapeDataInterface.h
--- a/Engine/Plugins/PCG/Source/PCG/Private/Compute/DataInterfaces/PCGLandscapeDataInterface.h
+++ b/Engine/Plugins/PCG/Source/PCG/Private/Compute/DataInterfaces/PCGLandscapeDataInterface.h
@@ -39,7 +39,7 @@ class UPCGLandscapeDataProvider : public UComputeDataProvider
 	GENERATED_BODY()
 
 public:
-	void Initialize(ALandscape* InLandscape, const FBox& Bounds);
+	void Init(ALandscape* InLandscape, const FBox& Bounds);
 
 	//~ Begin UComputeDataProvider Interface
 	FComputeDataProviderRenderProxy* GetRenderProxy() override;
diff --git a/Engine/Plugins/Runtime/ComputeFramework/Source/ComputeFramework/Private/ComputeGraph.cpp b/Engine/Plugins/Runtime/ComputeFramework/Source/ComputeFramework/Private/ComputeGraph.cpp
--- a/Engine/Plugins/Runtime/ComputeFramework/Source/ComputeFramework/Private/ComputeGraph.cpp
+++ b/Engine/Plugins/Runtime/ComputeFramework/Source/ComputeFramework/Private/ComputeGraph.cpp
@@ -245,6 +245,8 @@ bool UComputeGraph::ValidateProviders(TArray< TObjectPtr<UComputeDataProvider> >
 
 void UComputeGraph::CreateDataProviders(int32 InBindingIndex, TObjectPtr<UObject>& InBindingObject, TArray< TObjectPtr<UComputeDataProvider> >& InOutDataProviders) const
 {
+	TRACE_CPUPROFILER_EVENT_SCOPE(UComputeGraph::CreateDataProviders);
+
 	InOutDataProviders.SetNumZeroed(DataInterfaces.Num());
 
 	if (ensure(Bindings.IsValidIndex(InBindingIndex)) && (!InBindingObject || InBindingObject.IsA(Bindings[InBindingIndex])))
@@ -261,22 +263,21 @@ void UComputeGraph::CreateDataProviders(int32 InBindingIndex, TObjectPtr<UObject
 					// Gather which input/output bindings are connected in the graph.
 					uint64 InputMask = 0;
 					uint64 OutputMask = 0;
-					for (FComputeGraphEdge const& GraphEdge : GraphEdges)
+					GetDataInterfaceInputOutputMasks(DataInterfaceIndex, InputMask, OutputMask);
+
+					DataProvider = DataInterface->CreateDataProvider();
+
+					if (DataProvider)
 					{
-						if (GraphEdge.DataInterfaceIndex == DataInterfaceIndex)
-						{
-							if (GraphEdge.bKernelInput)
-							{
-								InputMask |= 1llu << GraphEdge.DataInterfaceBindingIndex;
-							}
-							else
-							{
-								OutputMask |= 1llu << GraphEdge.DataInterfaceBindingIndex;
-							}
-						}
+						DataProvider->Initialize(DataInterface, InBindingObject, InputMask, InputMask);
 					}
-
-					DataProvider = DataInterface->CreateDataProvider(InBindingObject, InputMask, OutputMask);
+					PRAGMA_DISABLE_DEPRECATION_WARNINGS
+					else
+					{
+						// Legacy fall back - try to use previous CreateDataProvider path.
+						DataProvider = DataInterface->CreateDataProvider(InBindingObject, InputMask, OutputMask);
+					}
+					PRAGMA_ENABLE_DEPRECATION_WARNINGS
 				}
 
 				InOutDataProviders[DataInterfaceIndex] = DataProvider;
@@ -285,6 +286,32 @@ void UComputeGraph::CreateDataProviders(int32 InBindingIndex, TObjectPtr<UObject
 	}
 }
 
+void UComputeGraph::InitializeDataProviders(int32 InBindingIndex, TObjectPtr<UObject>& InBindingObject, TArray< TObjectPtr<UComputeDataProvider> >& InDataProviders) const
+{
+	TRACE_CPUPROFILER_EVENT_SCOPE(UComputeGraph::InitializeDataProviders);
+
+	if (ensure(Bindings.IsValidIndex(InBindingIndex)) && (!InBindingObject || InBindingObject.IsA(Bindings[InBindingIndex])))
+	{
+		for (int32 DataInterfaceIndex = 0; DataInterfaceIndex < DataInterfaces.Num(); ++DataInterfaceIndex)
+		{
+			if (ensure(DataInterfaceToBinding.IsValidIndex(DataInterfaceIndex)) && DataInterfaceToBinding[DataInterfaceIndex] == InBindingIndex
+				&& ensure(InDataProviders.IsValidIndex(DataInterfaceIndex)) && ensure(InDataProviders[DataInterfaceIndex]))
+			{
+				UComputeDataInterface const* DataInterface = DataInterfaces[DataInterfaceIndex];
+				if (DataInterface != nullptr)
+				{
+					// Gather which input/output bindings are connected in the graph.
+					uint64 InputMask = 0;
+					uint64 OutputMask = 0;
+					GetDataInterfaceInputOutputMasks(DataInterfaceIndex, InputMask, OutputMask);
+
+					InDataProviders[DataInterfaceIndex]->Initialize(DataInterface, InBindingObject, InputMask, InputMask);
+				}
+			}
+		}
+	}
+}
+
 /** Compute Kernel compilation flags. */
 enum class EComputeKernelCompilationFlags
 {
@@ -518,6 +545,24 @@ void UComputeGraph::ReleaseRenderProxy(FComputeGraphRenderProxy* InRenderProxy)
 	}
 }
 
+void UComputeGraph::GetDataInterfaceInputOutputMasks(int32 InDataInterfaceIndex, uint64& OutInputMask, uint64& OutOutputMask) const
+{
+	for (FComputeGraphEdge const& GraphEdge : GraphEdges)
+	{
+		if (GraphEdge.DataInterfaceIndex == InDataInterfaceIndex)
+		{
+			if (GraphEdge.bKernelInput)
+			{
+				OutInputMask |= 1llu << GraphEdge.DataInterfaceBindingIndex;
+			}
+			else
+			{
+				OutOutputMask |= 1llu << GraphEdge.DataInterfaceBindingIndex;
+			}
+		}
+	}
+}
+
 #if WITH_EDITOR
 
 namespace
diff --git a/Engine/Plugins/Runtime/ComputeFramework/Source/ComputeFramework/Private/ComputeGraphInstance.cpp b/Engine/Plugins/Runtime/ComputeFramework/Source/ComputeFramework/Private/ComputeGraphInstance.cpp
--- a/Engine/Plugins/Runtime/ComputeFramework/Source/ComputeFramework/Private/ComputeGraphInstance.cpp
+++ b/Engine/Plugins/Runtime/ComputeFramework/Source/ComputeFramework/Private/ComputeGraphInstance.cpp
@@ -21,6 +21,25 @@ void FComputeGraphInstance::CreateDataProviders(UComputeGraph* InComputeGraph, i
 	}
 }
 
+void FComputeGraphInstance::InitializeDataProviders(UComputeGraph* InComputeGraph, int32 InBindingIndex, TObjectPtr<UObject> InBindingObject)
+{
+	if (InComputeGraph != nullptr)
+	{
+		InComputeGraph->InitializeDataProviders(InBindingIndex, InBindingObject, DataProviders);
+	}
+}
+
+void FComputeGraphInstance::ResetDataProviders()
+{
+	for (UComputeDataProvider* DataProvider : DataProviders)
+	{
+		if (DataProvider)
+		{
+			DataProvider->Reset();
+		}
+	}
+}
+
 void FComputeGraphInstance::DestroyDataProviders()
 {
 	DataProviders.Reset();
diff --git a/Engine/Plugins/Runtime/ComputeFramework/Source/ComputeFramework/Public/ComputeFramework/ComputeDataInterface.h b/Engine/Plugins/Runtime/ComputeFramework/Source/ComputeFramework/Public/ComputeFramework/ComputeDataInterface.h
--- a/Engine/Plugins/Runtime/ComputeFramework/Source/ComputeFramework/Public/ComputeFramework/ComputeDataInterface.h
+++ b/Engine/Plugins/Runtime/ComputeFramework/Source/ComputeFramework/Public/ComputeFramework/ComputeDataInterface.h
@@ -56,5 +56,8 @@ public:
 	/** Get class of UObject required to instantiate a UComputeDataProvider from this interface. */
 	virtual UClass* GetBindingType() const { return nullptr; }
 	/** Instantiate an associated UComputeDataProvider. */
+	virtual UComputeDataProvider* CreateDataProvider() const { return nullptr; }
+
+	UE_DEPRECATED(5.6, "Use the no-argument version together with UComputeDataProvider::Initialize()")
 	virtual UComputeDataProvider* CreateDataProvider(TObjectPtr<UObject> InBinding, uint64 InInputMask, uint64 InOutputMask) const { return nullptr; }
 };
diff --git a/Engine/Plugins/Runtime/ComputeFramework/Source/ComputeFramework/Public/ComputeFramework/ComputeDataProvider.h b/Engine/Plugins/Runtime/ComputeFramework/Source/ComputeFramework/Public/ComputeFramework/ComputeDataProvider.h
--- a/Engine/Plugins/Runtime/ComputeFramework/Source/ComputeFramework/Public/ComputeFramework/ComputeDataProvider.h
+++ b/Engine/Plugins/Runtime/ComputeFramework/Source/ComputeFramework/Public/ComputeFramework/ComputeDataProvider.h
@@ -7,6 +7,7 @@
 #include "Containers/StridedView.h"
 #include "ComputeDataProvider.generated.h"
 
+class UComputeDataInterface;
 class FComputeDataProviderRenderProxy;
 struct FComputeKernelPermutationVector;
 class FRDGBuilder;
@@ -23,6 +24,9 @@ class COMPUTEFRAMEWORK_API UComputeDataProvider : public UObject
 	GENERATED_BODY()
 
 public:
+	/** Set up the data provider from the given data interface and binding. */
+	virtual void Initialize(const UComputeDataInterface* InDataInterface, UObject* InBinding, uint64 InInputMask, uint64 InOutputMask) {}
+
 	/** Return false if the provider has not been fully initialized. */
 	UE_DEPRECATED(5.2, "Implement any validation in FComputeDataProviderRenderProxy::IsValid().")
 	virtual bool IsValid() const { return true; }
@@ -32,6 +36,9 @@ public:
 	 * Currently these are created and destroyed per frame by the FComputeGraphInstance.
 	 */
 	virtual FComputeDataProviderRenderProxy* GetRenderProxy() { return nullptr; }
+
+	/** Reset state and release any held resources. */
+	virtual void Reset() {}
 };
 
 /**
diff --git a/Engine/Plugins/Runtime/ComputeFramework/Source/ComputeFramework/Public/ComputeFramework/ComputeGraph.h b/Engine/Plugins/Runtime/ComputeFramework/Source/ComputeFramework/Public/ComputeFramework/ComputeGraph.h
--- a/Engine/Plugins/Runtime/ComputeFramework/Source/ComputeFramework/Public/ComputeFramework/ComputeGraph.h
+++ b/Engine/Plugins/Runtime/ComputeFramework/Source/ComputeFramework/Public/ComputeFramework/ComputeGraph.h
@@ -109,12 +109,14 @@ public:
 	bool ValidateGraph(FString* OutErrors = nullptr);
 
 	/** 
-	 * Create UComputeDataProvider objects to match the current UComputeDataInterface objects. 
+	 * Create UComputeDataProvider objects to match the current UComputeDataInterface objects and initialize them. 
 	 * We attempt to setup bindings from the InBindingObjects.
 	 * The caller is responsible for any data provider binding not handled by the default behavior.
 	 */
 	void CreateDataProviders(int32 InBindingIndex, TObjectPtr<UObject>& InBindingObject, TArray< TObjectPtr<UComputeDataProvider> >& InOutDataProviders) const;
 
+	/** Initialize the data providers. */
+	void InitializeDataProviders(int32 InBindingIndex, TObjectPtr<UObject>& InBindingObject, TArray< TObjectPtr<UComputeDataProvider> >& InDataProviders) const;
 
 	/** Returns true if there is a valid DataProvider entry for each of our DataInterfaces. */
 	bool ValidateProviders(TArray< TObjectPtr<UComputeDataProvider> > const& DataProviders) const;
@@ -156,6 +158,8 @@ private:
 	FComputeGraphRenderProxy* CreateRenderProxy() const;
 	/** Release the render proxy. */
 	void ReleaseRenderProxy(FComputeGraphRenderProxy* InProxy) const;
+	/** Returns input and output masks for the given data interface. */
+	void GetDataInterfaceInputOutputMasks(int32 InDataInterfaceIndex, uint64& OutInputMask, uint64& OutOutputMask) const;
 
 #if WITH_EDITOR
 	/** Build the HLSL source for a kernel with its linked data interfaces. */
diff --git a/Engine/Plugins/Runtime/ComputeFramework/Source/ComputeFramework/Public/ComputeFramework/ComputeGraphInstance.h b/Engine/Plugins/Runtime/ComputeFramework/Source/ComputeFramework/Public/ComputeFramework/ComputeGraphInstance.h
--- a/Engine/Plugins/Runtime/ComputeFramework/Source/ComputeFramework/Public/ComputeFramework/ComputeGraphInstance.h
+++ b/Engine/Plugins/Runtime/ComputeFramework/Source/ComputeFramework/Public/ComputeFramework/ComputeGraphInstance.h
@@ -26,11 +26,20 @@ public:
 	void SetGraphSortPriority(uint8 InPriority) { GraphSortPriority = InPriority; }
 
 	/** 
-	 * Create the Data Provider objects for a single binding of a ComputeGraph. 
+	 * Create and initialize the Data Provider objects for a single binding of a ComputeGraph. 
 	 * The type of binding object is expected to match the associated Binding on the UComputeGraph.
 	 */
 	void CreateDataProviders(UComputeGraph* InComputeGraph, int32 InBindingIndex, TObjectPtr<UObject> InBindingObject);
 
+	/**
+	 * Initialize the Data Provider objects for a single binding of a ComputeGraph.
+	 * The type of binding object is expected to match the associated Binding on the UComputeGraph.
+	 */
+	void InitializeDataProviders(UComputeGraph* InComputeGraph, int32 InBindingIndex, TObjectPtr<UObject> InBindingObject);
+
+	/** Clear the state within the Data Provider objects. */
+	void ResetDataProviders();
+
 	/** Create the Data Provider objects. */
 	void DestroyDataProviders();
 
