bump compute kernel version to clear out corrupt cache entries in LaneStrate project

#rnx

diff --git a/Engine/Plugins/Runtime/ComputeFramework/Source/ComputeFramework/Private/ComputeFramework/ComputeKernelDerivedDataVersion.h b/Engine/Plugins/Runtime/ComputeFramework/Source/ComputeFramework/Private/ComputeFramework/ComputeKernelDerivedDataVersion.h
--- a/Engine/Plugins/Runtime/ComputeFramework/Source/ComputeFramework/Private/ComputeFramework/ComputeKernelDerivedDataVersion.h
+++ b/Engine/Plugins/Runtime/ComputeFramework/Source/ComputeFramework/Private/ComputeFramework/ComputeKernelDerivedDataVersion.h
@@ -5,4 +5,4 @@
 #include "CoreMinimal.h"
 
 // In case of merge conflicts with DDC versions, you MUST generate a new GUID and set this new guid as version
-#define COMPUTEKERNEL_DERIVEDDATA_VER	TEXT("9222DB83-5B24-4401-8E40-D5A995931D20")
+#define COMPUTEKERNEL_DERIVEDDATA_VER	TEXT("A8B8CDC1-862C-4D85-BA8B-CD6C3FFF32DD")
