ComputeFramework: Cancel any outstanding compilation before destroying a ComputeGraph.
Fixes race condition crash where a ComputeGraph is released when still compiling, and the compiler job still has a pointer back to the ShaderMetadata owned by the ComputeGraph.
The ShaderMetadata is a shared pointer but ownership is more complex than that because the compute framework builds the shader metadata at runtime from multiple parts that point at one another, and ownership of these parts is handled in another object owned by the compute graph.
Another solution that I tried was blocking on compilation completion which is much simpler with fewer changes in the compiler manager. But cancellation should block less.
#rnx

diff --git a/Engine/Plugins/Runtime/ComputeFramework/Source/ComputeFramework/Private/ComputeFramework/ComputeKernelShaderCompilationManager.h b/Engine/Plugins/Runtime/ComputeFramework/Source/ComputeFramework/Private/ComputeFramework/ComputeKernelShaderCompilationManager.h
--- a/Engine/Plugins/Runtime/ComputeFramework/Source/ComputeFramework/Private/ComputeFramework/ComputeKernelShaderCompilationManager.h
+++ b/Engine/Plugins/Runtime/ComputeFramework/Source/ComputeFramework/Private/ComputeFramework/ComputeKernelShaderCompilationManager.h
@@ -41,6 +41,7 @@ public:
 	void Tick(float DeltaSeconds);
 	void AddJobs(TArray<FShaderCommonCompileJobPtr> InNewJobs);
 	void FinishCompilation(const TCHAR* InKernelName, const TArray<int32>& ShaderMapIdsToFinishCompiling);
+	void CancelCompilation(const TCHAR* InKernelName, const TArray<int32>& ShaderMapIdsToFinishCompiling);
 
 private:
 	void ProcessAsyncResults();
diff --git a/Engine/Plugins/Runtime/ComputeFramework/Source/ComputeFramework/Private/ComputeGraph.cpp b/Engine/Plugins/Runtime/ComputeFramework/Source/ComputeFramework/Private/ComputeGraph.cpp
--- a/Engine/Plugins/Runtime/ComputeFramework/Source/ComputeFramework/Private/ComputeGraph.cpp
+++ b/Engine/Plugins/Runtime/ComputeFramework/Source/ComputeFramework/Private/ComputeGraph.cpp
@@ -42,8 +42,19 @@ UComputeGraph::UComputeGraph(FVTableHelper& Helper)
 {
 }
 
-UComputeGraph::~UComputeGraph()
+void UComputeGraph::BeginDestroy()
 {
+	Super::BeginDestroy();
+	
+#if WITH_EDITORONLY_DATA
+	// Cancel any outstanding compilation jobs. These have a pointer to resource ShaderMetadata about to be deleted.
+	for (FComputeKernelResourceSet& KernelResource : KernelResources)
+	{
+		KernelResource.CancelCompilation();
+	}
+#endif
+
+	// Release on render thread. No need to wait on this before continuing destroy.
 	ReleaseRenderProxy(RenderProxy);
 }
 
@@ -1198,3 +1209,16 @@ void UComputeGraph::FComputeKernelResourceSet::ProcessSerializedShaderMaps()
 	LoadedKernelResources.Reset();
 #endif
 }
+
+#if WITH_EDITORONLY_DATA
+void UComputeGraph::FComputeKernelResourceSet::CancelCompilation()
+{
+	for (int32 FeatureLevel = 0; FeatureLevel < ERHIFeatureLevel::Num; ++FeatureLevel)
+	{
+		if (KernelResourcesByFeatureLevel[FeatureLevel].IsValid())
+		{
+			KernelResourcesByFeatureLevel[FeatureLevel]->CancelCompilation();
+		}
+	}
+}
+#endif
diff --git a/Engine/Plugins/Runtime/ComputeFramework/Source/ComputeFramework/Private/ComputeKernelShaderCompilationManager.cpp b/Engine/Plugins/Runtime/ComputeFramework/Source/ComputeFramework/Private/ComputeKernelShaderCompilationManager.cpp
--- a/Engine/Plugins/Runtime/ComputeFramework/Source/ComputeFramework/Private/ComputeKernelShaderCompilationManager.cpp
+++ b/Engine/Plugins/Runtime/ComputeFramework/Source/ComputeFramework/Private/ComputeKernelShaderCompilationManager.cpp
@@ -414,4 +414,16 @@ void FComputeKernelShaderCompilationManager::FinishCompilation(const TCHAR* InKe
 	check(ComputeKernelShaderMapJobs.Num() == 0);
 }
 
+void FComputeKernelShaderCompilationManager::CancelCompilation(const TCHAR* InKernelName, const TArray<int32>& ShaderMapIdsToCancelCompiling)
+{
+	check(!FPlatformProperties::RequiresCookedData());
+	GShaderCompilingManager->CancelCompilation(NULL, ShaderMapIdsToCancelCompiling);
+
+	for (int32 ShaderMapId : ShaderMapIdsToCancelCompiling)
+	{
+		JobQueue.RemoveAll([ShaderMapId](const FShaderCommonCompileJobPtr Job) { return Job->Id == ShaderMapId; });
+		ComputeKernelShaderMapJobs.Remove(ShaderMapId);
+	}
+}
+
 #endif // WITH_EDITOR
diff --git a/Engine/Plugins/Runtime/ComputeFramework/Source/ComputeFramework/Private/ComputeKernelShared.cpp b/Engine/Plugins/Runtime/ComputeFramework/Source/ComputeFramework/Private/ComputeKernelShared.cpp
--- a/Engine/Plugins/Runtime/ComputeFramework/Source/ComputeFramework/Private/ComputeKernelShared.cpp
+++ b/Engine/Plugins/Runtime/ComputeFramework/Source/ComputeFramework/Private/ComputeKernelShared.cpp
@@ -74,6 +74,10 @@ void FComputeKernelResource::CancelCompilation()
 {
 	if (IsInGameThread())
 	{
+		TArray<int32> ShaderMapIdsToFinish;
+		GetShaderMapIDsWithUnfinishedCompilation(ShaderMapIdsToFinish);
+		GComputeKernelShaderCompilationManager.CancelCompilation(*GetFriendlyName(), ShaderMapIdsToFinish);
+
 		FComputeKernelShaderMap::RemovePending(this);
 		OutstandingCompileShaderMapIds.Empty();
 
diff --git a/Engine/Plugins/Runtime/ComputeFramework/Source/ComputeFramework/Public/ComputeFramework/ComputeGraph.h b/Engine/Plugins/Runtime/ComputeFramework/Source/ComputeFramework/Public/ComputeFramework/ComputeGraph.h
--- a/Engine/Plugins/Runtime/ComputeFramework/Source/ComputeFramework/Public/ComputeFramework/ComputeGraph.h
+++ b/Engine/Plugins/Runtime/ComputeFramework/Source/ComputeFramework/Public/ComputeFramework/ComputeGraph.h
@@ -97,7 +97,6 @@ public:
 	UComputeGraph();
 	UComputeGraph(const FObjectInitializer& ObjectInitializer);
 	UComputeGraph(FVTableHelper& Helper);
-	virtual ~UComputeGraph();
 
 	/** Called each time that a single kernel shader compilation is completed. */
 	virtual void OnKernelCompilationComplete(int32 InKernelIndex, FComputeKernelCompileResults const& InCompileResults) {}
@@ -140,6 +139,7 @@ protected:
 	//~ Begin UObject Interface.
 	void Serialize(FArchive& Ar) override;
 	void PostLoad() override;
+	void BeginDestroy() override;
 #if WITH_EDITOR
 	void OnCookEvent(UE::Cook::ECookEvent CookEvent, UE::Cook::FCookEventContext& CookContext) override;
 	void BeginCacheForCookedPlatformData(ITargetPlatform const* TargetPlatform) override;
@@ -219,6 +219,10 @@ private:
 		void Serialize(FArchive& Ar);
 		/** Apply shader maps found in Serialize(). Call this from PostLoad(). */
 		void ProcessSerializedShaderMaps();
+
+#if WITH_EDITORONLY_DATA
+		void CancelCompilation();
+#endif
 	};
 
 	/** Kernel resources stored with the same indexing as the KernelInvocations array. */
