Minor parameter name change that got flagged in code review but missed from last submit.
#rnx

diff --git a/Engine/Plugins/Runtime/ComputeFramework/Source/ComputeFramework/Private/ComputeFramework/ComputeKernelShaderCompilationManager.h b/Engine/Plugins/Runtime/ComputeFramework/Source/ComputeFramework/Private/ComputeFramework/ComputeKernelShaderCompilationManager.h
--- a/Engine/Plugins/Runtime/ComputeFramework/Source/ComputeFramework/Private/ComputeFramework/ComputeKernelShaderCompilationManager.h
+++ b/Engine/Plugins/Runtime/ComputeFramework/Source/ComputeFramework/Private/ComputeFramework/ComputeKernelShaderCompilationManager.h
@@ -41,7 +41,7 @@ public:
 	void Tick(float DeltaSeconds);
 	void AddJobs(TArray<FShaderCommonCompileJobPtr> InNewJobs);
 	void FinishCompilation(const TCHAR* InKernelName, const TArray<int32>& ShaderMapIdsToFinishCompiling);
-	void CancelCompilation(const TCHAR* InKernelName, const TArray<int32>& ShaderMapIdsToFinishCompiling);
+	void CancelCompilation(const TCHAR* InKernelName, const TArray<int32>& ShaderMapIdsToCancelCompiling);
 
 private:
 	void ProcessAsyncResults();
