[ComputeFramework] Only do cancellation notification if a compilation was really cancelled.

#rnx

diff --git a/Engine/Plugins/Runtime/ComputeFramework/Source/ComputeFramework/Private/ComputeKernelShared.cpp b/Engine/Plugins/Runtime/ComputeFramework/Source/ComputeFramework/Private/ComputeKernelShared.cpp
--- a/Engine/Plugins/Runtime/ComputeFramework/Source/ComputeFramework/Private/ComputeKernelShared.cpp
+++ b/Engine/Plugins/Runtime/ComputeFramework/Source/ComputeFramework/Private/ComputeKernelShared.cpp
@@ -81,8 +81,11 @@ void FComputeKernelResource::CancelCompilation()
 		FComputeKernelShaderMap::RemovePending(this);
 		OutstandingCompileShaderMapIds.Empty();
 
-		FString Message = FString::Printf(TEXT("%s: Compilation cancelled."), *GetFriendlyName());
-		NotifyCompilationFinished(Message);
+		if (!ShaderMapIdsToFinish.IsEmpty())
+		{
+			FString Message = FString::Printf(TEXT("%s: Compilation cancelled."), *GetFriendlyName());
+			NotifyCompilationFinished(Message);
+		}
 	}
 }
 #endif // WITH_EDITOR
