renable compile on load

diff --git a/Engine/Plugins/Animation/DeformerGraph/Source/OptimusCore/Private/OptimusDeformer.cpp b/Engine/Plugins/Animation/DeformerGraph/Source/OptimusCore/Private/OptimusDeformer.cpp
--- a/Engine/Plugins/Animation/DeformerGraph/Source/OptimusCore/Private/OptimusDeformer.cpp
+++ b/Engine/Plugins/Animation/DeformerGraph/Source/OptimusCore/Private/OptimusDeformer.cpp
@@ -3452,7 +3452,7 @@ void UOptimusDeformer::PostLoad()
 	if (Status != EOptimusDeformerStatus::HasErrors &&
 		ObjectVersion < FOptimusObjectVersion::LatestVersion)
 	{
-		// Compile();
+		Compile();
 	}
 }
 
diff --git a/Engine/Plugins/Runtime/ComputeFramework/Source/ComputeFramework/Private/ComputeGraph.cpp b/Engine/Plugins/Runtime/ComputeFramework/Source/ComputeFramework/Private/ComputeGraph.cpp
--- a/Engine/Plugins/Runtime/ComputeFramework/Source/ComputeFramework/Private/ComputeGraph.cpp
+++ b/Engine/Plugins/Runtime/ComputeFramework/Source/ComputeFramework/Private/ComputeGraph.cpp
@@ -346,29 +346,17 @@ enum class EComputeKernelCompilationFlags
 
 void UComputeGraph::UpdateResources()
 {
-	bool bAllowRendering = true;
-	
 #if WITH_EDITOR
 	uint32 CompilationFlags = (uint32)EComputeKernelCompilationFlags::ApplyCompletedShaderMapForRendering;
 
-	// Rendering is likely not allowed during processes like cooking
-	if (IsRunningCommandlet() && !IsAllowCommandletRendering())
-	{
-		bAllowRendering = false;
-	}
-	
-	CompilationFlags |= bAllowRendering ? (uint32)EComputeKernelCompilationFlags::Synchronous : (uint32)EComputeKernelCompilationFlags::None;
+	bool bIsRenderCommandlet = IsRunningCommandlet() && IsAllowCommandletRendering();
+	CompilationFlags |= bIsRenderCommandlet ? (uint32)EComputeKernelCompilationFlags::Synchronous : (uint32)EComputeKernelCompilationFlags::None;
 
 	CacheResourceShadersForRendering(CompilationFlags);
 #endif
 
 	ReleaseRenderProxy(RenderProxy);
-	
-	// Rendering is disabled, so no need to create the render proxy
-	if (bAllowRendering)
-	{
-		RenderProxy = CreateRenderProxy();
-	}
+	RenderProxy = CreateRenderProxy();
 }
 
 bool UComputeGraph::HasKernelResourcesPendingShaderCompilation() const
@@ -458,6 +446,12 @@ void UComputeGraph::BuildShaderPermutationVectors(TArray<FComputeKernelPermutati
 
 FComputeGraphRenderProxy* UComputeGraph::CreateRenderProxy() const
 {
+	// Rendering is disabled, so no need to create the render proxy
+	if (!FApp::CanEverRender())
+	{
+		return nullptr;
+	}
+	
 	FComputeGraphRenderProxy* Proxy = new FComputeGraphRenderProxy;
 	Proxy->GraphName = GetFName();
 	Proxy->ShaderParameterMetadataAllocations = MakeUnique<FShaderParametersMetadataAllocations>();
