[Compute Framework] avoid sync compilation during SendEndOfAllFrameUpdate when running automation test. For automation test, it is better to compile the shaders upfront during postload instead of deferring to SendEndOfAllFrameUpdate to avoid marking render state dirty during SendEndOfAllFrameUpdate

diff --git a/Engine/Plugins/Runtime/ComputeFramework/Source/ComputeFramework/Private/ComputeFramework.cpp b/Engine/Plugins/Runtime/ComputeFramework/Source/ComputeFramework/Private/ComputeFramework.cpp
--- a/Engine/Plugins/Runtime/ComputeFramework/Source/ComputeFramework/Private/ComputeFramework.cpp
+++ b/Engine/Plugins/Runtime/ComputeFramework/Source/ComputeFramework/Private/ComputeFramework.cpp
@@ -53,7 +53,10 @@ namespace ComputeFramework
 	bool IsDeferredCompilation()
 	{
 #if WITH_EDITOR
-		return GComputeFrameworkDeferredCompilation != 0 && !(IsRunningCommandlet() && IsAllowCommandletRendering());
+		return
+			GComputeFrameworkDeferredCompilation != 0 &&
+			!(IsRunningCommandlet() && IsAllowCommandletRendering()) &&
+			!GIsAutomationTesting; // Making sure shader is ready before running the test, this helps avoid sync compilation during SendAllEndOfFrameUpdates, which is not allowed
 #else
 		return false;
 #endif
diff --git a/Engine/Plugins/Runtime/ComputeFramework/Source/ComputeFramework/Private/ComputeGraph.cpp b/Engine/Plugins/Runtime/ComputeFramework/Source/ComputeFramework/Private/ComputeGraph.cpp
--- a/Engine/Plugins/Runtime/ComputeFramework/Source/ComputeFramework/Private/ComputeGraph.cpp
+++ b/Engine/Plugins/Runtime/ComputeFramework/Source/ComputeFramework/Private/ComputeGraph.cpp
@@ -109,7 +109,9 @@ void UComputeGraph::PostLoad()
 
 	if (!ComputeFramework::IsDeferredCompilation())
 	{
-		UpdateResources();
+		// Sync compile here so that any downstream tasks like automation tests don't need to wait for shader compilation
+		constexpr bool bSync = true;
+		UpdateResources(bSync);
 	}
 }
 
@@ -340,17 +342,14 @@ enum class EComputeKernelCompilationFlags
 
 	/* Replaces all instances of the shader with the newly compiled version. */
 	ApplyCompletedShaderMapForRendering = 1 << 2,
-
-	IsCooking = 1 << 3,
 };
 
-void UComputeGraph::UpdateResources()
+void UComputeGraph::UpdateResources(bool bSync)
 {
 #if WITH_EDITOR
 	uint32 CompilationFlags = (uint32)EComputeKernelCompilationFlags::ApplyCompletedShaderMapForRendering;
 
-	bool bIsRenderCommandlet = IsRunningCommandlet() && IsAllowCommandletRendering();
-	CompilationFlags |= bIsRenderCommandlet ? (uint32)EComputeKernelCompilationFlags::Synchronous : (uint32)EComputeKernelCompilationFlags::None;
+	CompilationFlags |= bSync ? (uint32)EComputeKernelCompilationFlags::Synchronous : (uint32)EComputeKernelCompilationFlags::None;
 
 	CacheResourceShadersForRendering(CompilationFlags);
 #endif
@@ -910,10 +909,8 @@ void UComputeGraph::CacheShadersForResource(
 	uint32 CompilationFlags,
 	FComputeKernelResource* KernelResource)
 {
-	bool bCooking = (CompilationFlags & uint32(EComputeKernelCompilationFlags::IsCooking)) != 0;
-
 	const bool bIsDefault = (KernelResource->GetKernelFlags() & uint32(EComputeKernelFlags::IsDefaultKernel)) != 0;
-	if (!GIsEditor || GIsAutomationTesting || bIsDefault || bCooking)
+	if (bIsDefault)
 	{
 		CompilationFlags |= uint32(EComputeKernelCompilationFlags::Synchronous);
 	}
@@ -1015,7 +1012,7 @@ void UComputeGraph::BeginCacheForCookedPlatformData(ITargetPlatform const* Targe
 					ShaderParameterMetadata,
 					GetOutermost()->GetFName());
 
-				const uint32 CompilationFlags = uint32(EComputeKernelCompilationFlags::IsCooking);
+				const uint32 CompilationFlags = uint32(EComputeKernelCompilationFlags::Synchronous);
 				CacheShadersForResource(ShaderPlatform, TargetPlatform, CompilationFlags, KernelResource.Get());
 
 				Resources.Add(MoveTemp(KernelResource));
diff --git a/Engine/Plugins/Runtime/ComputeFramework/Source/ComputeFramework/Public/ComputeFramework/ComputeGraph.h b/Engine/Plugins/Runtime/ComputeFramework/Source/ComputeFramework/Public/ComputeFramework/ComputeGraph.h
--- a/Engine/Plugins/Runtime/ComputeFramework/Source/ComputeFramework/Public/ComputeFramework/ComputeGraph.h
+++ b/Engine/Plugins/Runtime/ComputeFramework/Source/ComputeFramework/Public/ComputeFramework/ComputeGraph.h
@@ -125,9 +125,9 @@ public:
 
 	/**
 	 * Call after changing the graph to build the graph resources for rendering.
-	 * This will trigger any required shader compilation.
+	 * This will trigger any required shader compilation. By default it does async compilation
 	 */
-	COMPUTEFRAMEWORK_API void UpdateResources();
+	COMPUTEFRAMEWORK_API void UpdateResources(bool bSync = false);
 
 	/**
 	 * Shader compilations triggered by UpdateResources() are async,
