[PCG] Add toggle to GPU nodes to trigger render captures

[RN] Added toggle to GPU nodes to trigger render captures for debugging (editor only)

diff --git a/Engine/Plugins/PCG/Source/PCG/Private/Compute/Elements/PCGComputeGraphElement.cpp b/Engine/Plugins/PCG/Source/PCG/Private/Compute/Elements/PCGComputeGraphElement.cpp
--- a/Engine/Plugins/PCG/Source/PCG/Private/Compute/Elements/PCGComputeGraphElement.cpp
+++ b/Engine/Plugins/PCG/Source/PCG/Private/Compute/Elements/PCGComputeGraphElement.cpp
@@ -226,6 +226,11 @@ bool FPCGComputeGraphElement::ExecuteInternal(FPCGContext* InContext) const
 			Context->ComputeGraphInstance->InitializeDataProviders(Context->ComputeGraph.Get(), 0, Context->DataBinding.Get());
 		}
 
+#if WITH_EDITOR
+		const bool bComponentInspected = Context->GetStack() && Graph && Graph->IsInspecting() && Graph->GetInspectedStack() == *Context->GetStack();
+		Context->ComputeGraphInstance->SetRenderCapturesEnabled(bComponentInspected);
+#endif
+
 		Context->bComputeGraphInstanceInitialized = true;
 
 		// Register all providers running async operations.
diff --git a/Engine/Plugins/PCG/Source/PCG/Private/Graph/PCGGraphCompilerGPU.cpp b/Engine/Plugins/PCG/Source/PCG/Private/Graph/PCGGraphCompilerGPU.cpp
--- a/Engine/Plugins/PCG/Source/PCG/Private/Graph/PCGGraphCompilerGPU.cpp
+++ b/Engine/Plugins/PCG/Source/PCG/Private/Graph/PCGGraphCompilerGPU.cpp
@@ -358,6 +358,11 @@ void FPCGGraphCompilerGPU::ExpandGPUNodeKernelsToTasks(
 			KernelInputs.Reset();
 			Kernel->GetInputPins(KernelInputs);
 
+			if (Settings->bTriggerRenderCapture)
+			{
+				Kernel->KernelFlags |= (int32)EComputeKernelFlags::TriggerRenderCapture;
+			}
+
 			// Wire inputs
 			for (const FPCGPinProperties& KernelInputProps : KernelInputs)
 			{
diff --git a/Engine/Plugins/PCG/Source/PCG/Public/PCGGraph.h b/Engine/Plugins/PCG/Source/PCG/Public/PCGGraph.h
--- a/Engine/Plugins/PCG/Source/PCG/Public/PCGGraph.h
+++ b/Engine/Plugins/PCG/Source/PCG/Public/PCGGraph.h
@@ -428,12 +428,14 @@ public:
 	void SetCommentNodes(TArray<FPCGGraphCommentNodeData> InNodes) { CommentNodes = std::move(InNodes); }
 	UE_API void RemoveCommentNode(const FGuid& InNodeGUID);
 
+	UE_API void RemoveExtraEditorNode(const UObject* InNode);
+
 	bool IsInspecting() const { return bIsInspecting; }
 	void EnableInspection(const FPCGStack& InInspectedStack) { bIsInspecting = true; InspectedStack = InInspectedStack; }
 	void DisableInspection() { bIsInspecting = false; InspectedStack = FPCGStack(); }
 	bool DebugFlagAppliesToIndividualComponents() const { return bDebugFlagAppliesToIndividualComponents; }
-	UE_API void RemoveExtraEditorNode(const UObject* InNode);
-
+	const FPCGStack& GetInspectedStack() const { return InspectedStack; }
+	
 	/** Instruct the graph compiler to cache the relevant permutations of this graph. */
 	UE_API bool PrimeGraphCompilationCache();
 
diff --git a/Engine/Plugins/PCG/Source/PCG/Public/PCGSettings.h b/Engine/Plugins/PCG/Source/PCG/Public/PCGSettings.h
--- a/Engine/Plugins/PCG/Source/PCG/Public/PCGSettings.h
+++ b/Engine/Plugins/PCG/Source/PCG/Public/PCGSettings.h
@@ -560,6 +560,12 @@ public:
 	UPROPERTY(EditAnywhere, Category = "GPU", AdvancedDisplay, meta = (Tooltip = "Size (in number of floats) of the shader debug print buffer.", EditCondition="bExecuteOnGPU && bPrintShaderDebugValues", EditConditionHides))
 	int DebugBufferSize = 16;
 
+#if WITH_EDITORONLY_DATA
+	/** [EXPERIMENTAL] Note: This variable is subject to change without deprecation. */
+	UPROPERTY(EditAnywhere, Transient, Category = "GPU", AdvancedDisplay, meta = (Tooltip = "Will trigger a render capture when this node executes and a debug object is selected in the graph editor. Transient and Editor-only. Render captures must be enabled (e.g. -AttachRenderDoc or -AttachPIX).", EditCondition = "bExecuteOnGPU", EditConditionHides))
+	bool bTriggerRenderCapture = false;
+#endif
+
 private:
 	/** Calculate Crc for these settings and save it. */
 	UE_API void CacheCrc();
diff --git a/Engine/Plugins/PCG/Source/PCGEditor/Private/PCGEditor.cpp b/Engine/Plugins/PCG/Source/PCGEditor/Private/PCGEditor.cpp
--- a/Engine/Plugins/PCG/Source/PCGEditor/Private/PCGEditor.cpp
+++ b/Engine/Plugins/PCG/Source/PCGEditor/Private/PCGEditor.cpp
@@ -255,7 +255,7 @@ void FPCGEditor::SetStackBeingInspected(const FPCGStack& FullStack)
 			PCGGraphBeingEdited->DisableInspection();
 		}
 
-		PCGGraphBeingEdited->EnableInspection(StackBeingInspected);
+		PCGGraphBeingEdited->EnableInspection(FullStack);
 	}
 
 	PCGSourceBeingInspected = NewComponent;
diff --git a/Engine/Plugins/Runtime/ComputeFramework/Source/ComputeFramework/Private/ComputeFramework/ComputeGraphRenderProxy.h b/Engine/Plugins/Runtime/ComputeFramework/Source/ComputeFramework/Private/ComputeFramework/ComputeGraphRenderProxy.h
--- a/Engine/Plugins/Runtime/ComputeFramework/Source/ComputeFramework/Private/ComputeFramework/ComputeGraphRenderProxy.h
+++ b/Engine/Plugins/Runtime/ComputeFramework/Source/ComputeFramework/Private/ComputeFramework/ComputeGraphRenderProxy.h
@@ -42,6 +42,8 @@ public:
 		int32 ExecutionProviderIndex = -1;
 		/** Whether the kernel can combine multiple supports combined sub-invocations. */
 		bool bSupportsUnifiedDispatch = false;
+		/** Whether a render capture should be triggered in the scope around this kernel invocation. */
+		bool bTriggerRenderCapture = false;
 	};
 
 	/** Friendly name for the owner graph. */
diff --git a/Engine/Plugins/Runtime/ComputeFramework/Source/ComputeFramework/Private/ComputeFramework/ComputeGraphWorker.h b/Engine/Plugins/Runtime/ComputeFramework/Source/ComputeFramework/Private/ComputeFramework/ComputeGraphWorker.h
--- a/Engine/Plugins/Runtime/ComputeFramework/Source/ComputeFramework/Private/ComputeFramework/ComputeGraphWorker.h
+++ b/Engine/Plugins/Runtime/ComputeFramework/Source/ComputeFramework/Private/ComputeFramework/ComputeGraphWorker.h
@@ -28,6 +28,7 @@ public:
 		FComputeGraphRenderProxy const* InGraphRenderProxy, 
 		TArray<FComputeDataProviderRenderProxy*> InDataProviderRenderProxies,
 		FSimpleDelegate InFallbackDelegate,
+		bool bInEnableRenderCaptures,
 		const UObject* OwnerPointer = nullptr);
 
 	void Abort(const UObject* OwnerPointer);
@@ -61,7 +62,9 @@ private:
 		TArray<FComputeDataProviderRenderProxy*> DataProviderRenderProxies;
 		/** Render thread fallback logic for invocations that are invalid. */
 		FSimpleDelegate FallbackDelegate;
-		
+		/** Whether render captures are enabled. */
+		bool bEnableRenderCaptures = false;
+
 		~FGraphInvocation();
 	};
 
@@ -76,7 +79,7 @@ private:
 		 * Syncing is usually required between consecutive kernels in a graph.
 		 * So we schedule the first kernels of all the graphs, before all of the second kernels.
 		 * That reduces time sync time, at the expense of memory pressure for buffers that need to stay alive.
-		 * In futuer we may want to add a limit to the number of graphs in flight to avoid memory pressure.
+		 * In future we may want to add a limit to the number of graphs in flight to avoid memory pressure.
 		 */
 		union
 		{
@@ -93,6 +96,8 @@ private:
 		uint32 ShaderIndex : 15;
 		/** Track if this is Unified Dispatch. */
 		uint32 bIsUnified : 1;
+		/** Whether a render capture should be triggered in the scope around this submit. */
+		uint32 bTriggerRenderCapture : 1;
 	};
 
 	// These arrays could be local to FComputeGraphTaskWorker::SubmitWork() but we 
diff --git a/Engine/Plugins/Runtime/ComputeFramework/Source/ComputeFramework/Private/ComputeGraph.cpp b/Engine/Plugins/Runtime/ComputeFramework/Source/ComputeFramework/Private/ComputeGraph.cpp
--- a/Engine/Plugins/Runtime/ComputeFramework/Source/ComputeFramework/Private/ComputeGraph.cpp
+++ b/Engine/Plugins/Runtime/ComputeFramework/Source/ComputeFramework/Private/ComputeGraph.cpp
@@ -474,6 +474,7 @@ FComputeGraphRenderProxy* UComputeGraph::CreateRenderProxy() const
 			Invocation.KernelResource = KernelResource;
 			Invocation.ShaderParameterMetadata = BuildKernelShaderMetadata(KernelIndex, *Proxy->ShaderParameterMetadataAllocations);
 			Invocation.bSupportsUnifiedDispatch = true;
+			Invocation.bTriggerRenderCapture = !!(Kernel->KernelFlags & (int32)EComputeKernelFlags::TriggerRenderCapture);
 
 			// 1. Figure out which data interface is the execution data interface
 			for (FComputeGraphEdge const& GraphEdge : GraphEdges)
diff --git a/Engine/Plugins/Runtime/ComputeFramework/Source/ComputeFramework/Private/ComputeGraphInstance.cpp b/Engine/Plugins/Runtime/ComputeFramework/Source/ComputeFramework/Private/ComputeGraphInstance.cpp
--- a/Engine/Plugins/Runtime/ComputeFramework/Source/ComputeFramework/Private/ComputeGraphInstance.cpp
+++ b/Engine/Plugins/Runtime/ComputeFramework/Source/ComputeFramework/Private/ComputeGraphInstance.cpp
@@ -91,10 +91,10 @@ bool FComputeGraphInstance::EnqueueWork(UComputeGraph* InComputeGraph, FSceneInt
 	}
 
 	ENQUEUE_RENDER_COMMAND(ComputeFrameworkEnqueueExecutionCommand)(
-		[ComputeGraphWorker, InExecutionGroupName, InOwnerName, SortPriority = InGraphSortPriorityOffset + GraphSortPriority, GraphRenderProxy, MovedDataProviderRenderProxies = MoveTemp(DataProviderRenderProxies), InFallbackDelegate, InOwnerPointer](FRHICommandListImmediate& RHICmdList)
+		[ComputeGraphWorker, InExecutionGroupName, InOwnerName, SortPriority = InGraphSortPriorityOffset + GraphSortPriority, GraphRenderProxy, MovedDataProviderRenderProxies = MoveTemp(DataProviderRenderProxies), InFallbackDelegate, InOwnerPointer, bEnableRenderCaptures=this->bEnableRenderCaptures](FRHICommandListImmediate& RHICmdList)
 		{
 			// Compute graph scheduler will take ownership of the provider proxies.
-			ComputeGraphWorker->Enqueue(InExecutionGroupName, InOwnerName, SortPriority, GraphRenderProxy, MovedDataProviderRenderProxies, InFallbackDelegate, InOwnerPointer);
+			ComputeGraphWorker->Enqueue(InExecutionGroupName, InOwnerName, SortPriority, GraphRenderProxy, MovedDataProviderRenderProxies, InFallbackDelegate, bEnableRenderCaptures, InOwnerPointer);
 		});
 
 	return true;
diff --git a/Engine/Plugins/Runtime/ComputeFramework/Source/ComputeFramework/Private/ComputeGraphWorker.cpp b/Engine/Plugins/Runtime/ComputeFramework/Source/ComputeFramework/Private/ComputeGraphWorker.cpp
--- a/Engine/Plugins/Runtime/ComputeFramework/Source/ComputeFramework/Private/ComputeGraphWorker.cpp
+++ b/Engine/Plugins/Runtime/ComputeFramework/Source/ComputeFramework/Private/ComputeGraphWorker.cpp
@@ -41,6 +41,7 @@ void FComputeGraphTaskWorker::Enqueue(
 	FComputeGraphRenderProxy const* InGraphRenderProxy, 
 	TArray<FComputeDataProviderRenderProxy*> InDataProviderRenderProxies, 
 	FSimpleDelegate InFallbackDelegate,
+	bool bInEnableRenderCaptures,
 	const UObject* InOwnerPointer)
 {
 	FGraphInvocation& GraphInvocation = GraphInvocationsPerGroup.FindOrAdd(InExecutionGroupName).AddDefaulted_GetRef();
@@ -50,6 +51,7 @@ void FComputeGraphTaskWorker::Enqueue(
 	GraphInvocation.GraphRenderProxy = InGraphRenderProxy;
 	GraphInvocation.DataProviderRenderProxies = MoveTemp(InDataProviderRenderProxies);
 	GraphInvocation.FallbackDelegate = InFallbackDelegate;
+	GraphInvocation.bEnableRenderCaptures = bInEnableRenderCaptures;
 }
 
 void FComputeGraphTaskWorker::Abort(const UObject* InOwnerPointer)
@@ -102,7 +104,8 @@ void FComputeGraphTaskWorker::SubmitWork(FRDGBuilder& GraphBuilder, FName InExec
 	RDG_EVENT_SCOPE_STAT(GraphBuilder, ComputeFramework_ExecuteBatches, "ComputeFramework::ExecuteBatches");
 	RDG_GPU_STAT_SCOPE(GraphBuilder, ComputeFramework_ExecuteBatches);
 
-	RenderCaptureInterface::FScopedCapture RenderCapture(GTriggerGPUCaptureDispatches > 0, GraphBuilder, TEXT("FComputeGraphTaskWorker::SubmitWork"));
+	const bool bCaptureAllWork = GTriggerGPUCaptureDispatches > 0;
+	RenderCaptureInterface::FScopedCapture RenderCapture(bCaptureAllWork, GraphBuilder, TEXT("FComputeGraphTaskWorker::SubmitWork"));
 	GTriggerGPUCaptureDispatches = FMath::Max(GTriggerGPUCaptureDispatches - 1, 0);
 
 	// Currently poll readbacks once at end of frame.
@@ -149,6 +152,7 @@ void FComputeGraphTaskWorker::SubmitWork(FRDGBuilder& GraphBuilder, FName InExec
 			SubmitDesc.ShaderIndex = Shaders.Num();
 
 			FComputeGraphRenderProxy::FKernelInvocation const& KernelInvocation = GraphRenderProxy->KernelInvocations[KernelIndex];
+			SubmitDesc.bTriggerRenderCapture = GraphInvocation.bEnableRenderCaptures && KernelInvocation.bTriggerRenderCapture;
 
 			// Reset our scratch memory arrays.
 			PermutationIds.Reset();
@@ -248,6 +252,8 @@ void FComputeGraphTaskWorker::SubmitWork(FRDGBuilder& GraphBuilder, FName InExec
 		const int32 KernelIndex = SubmitDesc.KernelIndex;
 		FComputeGraphRenderProxy::FKernelInvocation const& KernelInvocation = GraphRenderProxy->KernelInvocations[KernelIndex];
 
+		RenderCaptureInterface::FScopedCapture SubmitRenderCapture(SubmitDesc.bTriggerRenderCapture && !bCaptureAllWork, GraphBuilder, TEXT("FComputeGraphTaskWorker::SubmitWork::SubmitDesc"));
+
 		RDG_EVENT_SCOPE(GraphBuilder, "%s:%s:%s", *GraphInvocation.OwnerName.ToString(), *GraphRenderProxy->GraphName.ToString(), *KernelInvocation.KernelName);
 
 		// Pre submit calls.
diff --git a/Engine/Plugins/Runtime/ComputeFramework/Source/ComputeFramework/Public/ComputeFramework/ComputeGraphInstance.h b/Engine/Plugins/Runtime/ComputeFramework/Source/ComputeFramework/Public/ComputeFramework/ComputeGraphInstance.h
--- a/Engine/Plugins/Runtime/ComputeFramework/Source/ComputeFramework/Public/ComputeFramework/ComputeGraphInstance.h
+++ b/Engine/Plugins/Runtime/ComputeFramework/Source/ComputeFramework/Public/ComputeFramework/ComputeGraphInstance.h
@@ -54,6 +54,8 @@ public:
 	/** Enqueue the ComputeGraph work. */
 	UE_API bool EnqueueWork(UComputeGraph* InComputeGraph, FSceneInterface const* InScene, FName InExecutionGroupName, FName InOwnerName, FSimpleDelegate InFallbackDelegate, UObject* InOwnerPointer = nullptr, uint8 InGraphSortPriorityOffset = 0);
 
+	UE_API void SetRenderCapturesEnabled(bool bInEnable) { bEnableRenderCaptures = bInEnable; }
+
 private:
 	/** The currently bound Data Provider objects. */
 	UPROPERTY(Transient)
@@ -61,6 +63,8 @@ private:
 
 	/** Priority used when sorting work. */
 	uint8 GraphSortPriority = 0;
+
+	bool bEnableRenderCaptures = false;
 };
 
 #undef UE_API
diff --git a/Engine/Plugins/Runtime/ComputeFramework/Source/ComputeFramework/Public/ComputeFramework/ComputeKernel.h b/Engine/Plugins/Runtime/ComputeFramework/Source/ComputeFramework/Public/ComputeFramework/ComputeKernel.h
--- a/Engine/Plugins/Runtime/ComputeFramework/Source/ComputeFramework/Public/ComputeFramework/ComputeKernel.h
+++ b/Engine/Plugins/Runtime/ComputeFramework/Source/ComputeFramework/Public/ComputeFramework/ComputeKernel.h
@@ -18,14 +18,19 @@ enum class EComputeKernelFlags : uint32
 	 * It also implies that this will be compiled synchronously. Other than a pass
 	 * through kernel, default shouldn't be used.
 	 */
-	IsDefaultKernel = 1 << 0, // KERNEL_FLAG(IS_DEFAULT_KERNEL)
+	IsDefaultKernel = 1 << 0,
 
 	/*
 	 * Promise from the author that all memory writes will be unique per shader
 	 * dispatch thread. i.e. ThreadX will be the only thread to write to MemoryY,
 	 * thus no synchronization is necessary by the compute graph.
 	 */
-	IsolatedMemoryWrites = 1 << 1, // KERNEL_FLAG(ISOLATED_MEMORY_WRITES)
+	IsolatedMemoryWrites = 1 << 1,
+
+	/**
+	 * Trigger render capture when this kernel dispatches, for debugging/testing purposes.
+	 */
+	TriggerRenderCapture = 1 << 2,
 };
 ENUM_CLASS_FLAGS(EComputeKernelFlags);
 
