ComputeFramework: Add support for graph build and test without requirement for DeformerGraph of PCG clients.
Added Dispatch and Buffer data interfaces.
Added UComputeGraphFromText that derives from UComputeGraph and can set up the ComputeGraph.
Currently this creates a hardcoded graph, but will be changed soon to read in a data driven graph setup.
Fixed UComputeGraphComponent to init providers on registration.
Call ResetTypeDeclaration() when filling out a FShaderParamTypeDefinition. Client code is currently forced to always do this before using the def which seems fragile.
Removed unused code.

diff --git a/Engine/Plugins/Runtime/ComputeFramework/Shaders/Private/ComputeDataInterfaceBuffer.ush b/Engine/Plugins/Runtime/ComputeFramework/Shaders/Private/ComputeDataInterfaceBuffer.ush
new file mode 100644
--- /dev/null
+++ b/Engine/Plugins/Runtime/ComputeFramework/Shaders/Private/ComputeDataInterfaceBuffer.ush
@@ -0,0 +1,225 @@
+// Copyright Epic Games, Inc. All Rights Reserved.
+
+// VALUE_TYPE is the exposed type that calling code expects to use.
+#define VALUE_TYPE {ValueType}
+#define VALUE_TYPE_STRIDE_BYTES {ValueTypeStride}
+// VALUE_TYPE_SUPPORTS_ATOMIC is set if we can enable atomic operations on the VALUE_TYPE.
+#define VALUE_TYPE_SUPPORTS_ATOMIC {SupportAtomic}
+// BUFFER_SPLIT_READ_WRITE is set if we can use an SRV to read the buffer. It is unset if we might need to read and write the buffer in the same kernel.
+#define BUFFER_SPLIT_READ_WRITE {SplitReadWrite}
+
+uint {DataInterfaceName}_BufferElementCount;
+#if BUFFER_SPLIT_READ_WRITE
+ByteAddressBuffer {DataInterfaceName}_BufferSRV;
+#endif
+RWByteAddressBuffer {DataInterfaceName}_BufferUAV;
+
+void ReadBufferSRV_{DataInterfaceName}(in uint Index, out uint ReturnValue)
+{
+	ReturnValue = {DataInterfaceName}_BufferSRV.Load(Index * VALUE_TYPE_STRIDE_BYTES);
+}
+void ReadBufferSRV_{DataInterfaceName}(in uint Index, out uint2 ReturnValue)
+{
+	ReturnValue = {DataInterfaceName}_BufferSRV.Load2(Index * VALUE_TYPE_STRIDE_BYTES);
+}
+void ReadBufferSRV_{DataInterfaceName}(in uint Index, out uint3 ReturnValue)
+{
+	ReturnValue = {DataInterfaceName}_BufferSRV.Load3(Index * VALUE_TYPE_STRIDE_BYTES);
+}
+void ReadBufferSRV_{DataInterfaceName}(in uint Index, out uint4 ReturnValue)
+{
+	ReturnValue = {DataInterfaceName}_BufferSRV.Load4(Index * VALUE_TYPE_STRIDE_BYTES);
+}
+void ReadBufferUAV_{DataInterfaceName}(in uint Index, out uint ReturnValue)
+{
+	ReturnValue = {DataInterfaceName}_BufferSRV.Load(Index * VALUE_TYPE_STRIDE_BYTES);
+}
+void ReadBufferUAV_{DataInterfaceName}(in uint Index, out uint2 ReturnValue)
+{
+	ReturnValue = {DataInterfaceName}_BufferSRV.Load2(Index * VALUE_TYPE_STRIDE_BYTES);
+}
+void ReadBufferUAV_{DataInterfaceName}(in uint Index, out uint3 ReturnValue)
+{
+	ReturnValue = {DataInterfaceName}_BufferSRV.Load3(Index * VALUE_TYPE_STRIDE_BYTES);
+}
+void ReadBufferUAV_{DataInterfaceName}(in uint Index, out uint4 ReturnValue)
+{
+	ReturnValue = {DataInterfaceName}_BufferSRV.Load4(Index * VALUE_TYPE_STRIDE_BYTES);
+}
+void WriteBuffer_{DataInterfaceName}(in uint Index, in uint Value)
+{
+	{DataInterfaceName}_BufferUAV.Store(Index * VALUE_TYPE_STRIDE_BYTES, Value);
+}
+void WriteBuffer_{DataInterfaceName}(in uint Index, in uint2 Value)
+{
+	{DataInterfaceName}_BufferUAV.Store2(Index * VALUE_TYPE_STRIDE_BYTES, Value);
+}
+void WriteBuffer_{DataInterfaceName}(in uint Index, in uint3 Value)
+{
+	{DataInterfaceName}_BufferUAV.Store3(Index * VALUE_TYPE_STRIDE_BYTES, Value);
+}
+void WriteBuffer_{DataInterfaceName}(in uint Index, in uint4 Value)
+{
+	{DataInterfaceName}_BufferUAV.Store4(Index * VALUE_TYPE_STRIDE_BYTES, Value);
+}
+
+void ReadBufferSRV_{DataInterfaceName}(in uint Index, out int ReturnValue)
+{
+	ReturnValue = asint({DataInterfaceName}_BufferSRV.Load(Index * VALUE_TYPE_STRIDE_BYTES));
+}
+void ReadBufferSRV_{DataInterfaceName}(in uint Index, out int2 ReturnValue)
+{
+	ReturnValue = asint({DataInterfaceName}_BufferSRV.Load2(Index * VALUE_TYPE_STRIDE_BYTES));
+}
+void ReadBufferSRV_{DataInterfaceName}(in uint Index, out int3 ReturnValue)
+{
+	ReturnValue = asint({DataInterfaceName}_BufferSRV.Load3(Index * VALUE_TYPE_STRIDE_BYTES));
+}
+void ReadBufferSRV_{DataInterfaceName}(in uint Index, out int4 ReturnValue)
+{
+	ReturnValue = asint({DataInterfaceName}_BufferSRV.Load4(Index * VALUE_TYPE_STRIDE_BYTES));
+}
+void ReadBufferUAV_{DataInterfaceName}(in uint Index, out int ReturnValue)
+{
+	ReturnValue = asint({DataInterfaceName}_BufferSRV.Load(Index * VALUE_TYPE_STRIDE_BYTES));
+}
+void ReadBufferUAV_{DataInterfaceName}(in uint Index, out int2 ReturnValue)
+{
+	ReturnValue = asint({DataInterfaceName}_BufferSRV.Load2(Index * VALUE_TYPE_STRIDE_BYTES));
+}
+void ReadBufferUAV_{DataInterfaceName}(in uint Index, out int3 ReturnValue)
+{
+	ReturnValue = asint({DataInterfaceName}_BufferSRV.Load3(Index * VALUE_TYPE_STRIDE_BYTES));
+}
+void ReadBufferUAV_{DataInterfaceName}(in uint Index, out int4 ReturnValue)
+{
+	ReturnValue = asint({DataInterfaceName}_BufferSRV.Load4(Index * VALUE_TYPE_STRIDE_BYTES));
+}
+void WriteBuffer_{DataInterfaceName}(in uint Index, in int Value)
+{
+	{DataInterfaceName}_BufferUAV.Store(Index * VALUE_TYPE_STRIDE_BYTES, asuint(Value));
+}
+void WriteBuffer_{DataInterfaceName}(in uint Index, in int2 Value)
+{
+	{DataInterfaceName}_BufferUAV.Store2(Index * VALUE_TYPE_STRIDE_BYTES, asuint(Value));
+}
+void WriteBuffer_{DataInterfaceName}(in uint Index, in int3 Value)
+{
+	{DataInterfaceName}_BufferUAV.Store3(Index * VALUE_TYPE_STRIDE_BYTES, asuint(Value));
+}
+void WriteBuffer_{DataInterfaceName}(in uint Index, in int4 Value)
+{
+	{DataInterfaceName}_BufferUAV.Store4(Index * VALUE_TYPE_STRIDE_BYTES, asuint(Value));
+}
+
+void ReadBufferSRV_{DataInterfaceName}(in uint Index, out float ReturnValue)
+{
+	ReturnValue = asfloat({DataInterfaceName}_BufferSRV.Load(Index * VALUE_TYPE_STRIDE_BYTES));
+}
+void ReadBufferSRV_{DataInterfaceName}(in uint Index, out float2 ReturnValue)
+{
+	ReturnValue = asfloat({DataInterfaceName}_BufferSRV.Load2(Index * VALUE_TYPE_STRIDE_BYTES));
+}
+void ReadBufferSRV_{DataInterfaceName}(in uint Index, out float3 ReturnValue)
+{
+	ReturnValue = asfloat({DataInterfaceName}_BufferSRV.Load3(Index * VALUE_TYPE_STRIDE_BYTES));
+}
+void ReadBufferSRV_{DataInterfaceName}(in uint Index, out float4 ReturnValue)
+{
+	ReturnValue = asfloat({DataInterfaceName}_BufferSRV.Load4(Index * VALUE_TYPE_STRIDE_BYTES));
+}
+void ReadBufferUAV_{DataInterfaceName}(in uint Index, out float ReturnValue)
+{
+	ReturnValue = asfloat({DataInterfaceName}_BufferSRV.Load(Index * VALUE_TYPE_STRIDE_BYTES));
+}
+void ReadBufferUAV_{DataInterfaceName}(in uint Index, out float2 ReturnValue)
+{
+	ReturnValue = asfloat({DataInterfaceName}_BufferSRV.Load2(Index * VALUE_TYPE_STRIDE_BYTES));
+}
+void ReadBufferUAV_{DataInterfaceName}(in uint Index, out float3 ReturnValue)
+{
+	ReturnValue = asfloat({DataInterfaceName}_BufferSRV.Load3(Index * VALUE_TYPE_STRIDE_BYTES));
+}
+void ReadBufferUAV_{DataInterfaceName}(in uint Index, out float4 ReturnValue)
+{
+	ReturnValue = asfloat({DataInterfaceName}_BufferSRV.Load4(Index * VALUE_TYPE_STRIDE_BYTES));
+}
+void WriteBuffer_{DataInterfaceName}(in uint Index, in float Value)
+{
+	{DataInterfaceName}_BufferUAV.Store(Index * VALUE_TYPE_STRIDE_BYTES, asuint(Value));
+}
+void WriteBuffer_{DataInterfaceName}(in uint Index, in float2 Value)
+{
+	{DataInterfaceName}_BufferUAV.Store2(Index * VALUE_TYPE_STRIDE_BYTES, asuint(Value));
+}
+void WriteBuffer_{DataInterfaceName}(in uint Index, in float3 Value)
+{
+	{DataInterfaceName}_BufferUAV.Store3(Index * VALUE_TYPE_STRIDE_BYTES, asuint(Value));
+}
+void WriteBuffer_{DataInterfaceName}(in uint Index, in float4 Value)
+{
+	{DataInterfaceName}_BufferUAV.Store4(Index * VALUE_TYPE_STRIDE_BYTES, asuint(Value));
+}
+
+uint ReadNumValues_{DataInterfaceName}()
+{
+	return {DataInterfaceName}_BufferElementCount;
+}
+
+VALUE_TYPE ReadValue_{DataInterfaceName}(uint Index)
+{
+	VALUE_TYPE ReturnValue;
+#if BUFFER_SPLIT_READ_WRITE
+	ReadBufferSRV_{DataInterfaceName}(Index, ReturnValue);
+#else
+	ReadBufferUAV_{DataInterfaceName}(Index, ReturnValue);
+#endif
+	return ReturnValue;
+}
+
+VALUE_TYPE ReadValueUAV_{DataInterfaceName}(uint Index)
+{
+	VALUE_TYPE ReturnValue;
+	ReadBufferUAV_{DataInterfaceName}(Index, ReturnValue);
+	return ReturnValue;
+}
+
+void WriteValue_{DataInterfaceName}(uint Index, VALUE_TYPE Value)
+{
+	WriteBuffer_{DataInterfaceName}(Index, Value);
+}
+
+VALUE_TYPE WriteAtomicAdd_{DataInterfaceName}(uint Index, VALUE_TYPE Value)
+{
+#if VALUE_TYPE_SUPPORTS_ATOMIC
+	VALUE_TYPE OldValue;
+	{DataInterfaceName}_BufferUAV.InterlockedAdd(Index * VALUE_TYPE_STRIDE_BYTES, Value, OldValue);
+	return OldValue;
+#endif
+	return 0;
+}
+
+VALUE_TYPE WriteAtomicMin_{DataInterfaceName}(uint Index, VALUE_TYPE Value)
+{
+#if VALUE_TYPE_SUPPORTS_ATOMIC
+	VALUE_TYPE OldValue;
+	{DataInterfaceName}_BufferUAV.InterlockedMin(Index * VALUE_TYPE_STRIDE_BYTES, Value, OldValue);
+	return OldValue;
+#endif
+	return 0;
+}
+
+VALUE_TYPE WriteAtomicMax_{DataInterfaceName}(uint Index, VALUE_TYPE Value)
+{
+#if VALUE_TYPE_SUPPORTS_ATOMIC
+	VALUE_TYPE OldValue;
+	{DataInterfaceName}_BufferUAV.InterlockedMax(Index * VALUE_TYPE_STRIDE_BYTES, Value, OldValue);
+	return OldValue;
+#endif
+	return 0;
+}
+
+#undef VALUE_TYPE
+#undef VALUE_TYPE_STRIDE_BYTES
+#undef VALUE_TYPE_SUPPORTS_ATOMIC
+#undef BUFFER_SPLIT_READ_WRITE
diff --git a/Engine/Plugins/Runtime/ComputeFramework/Shaders/Private/ComputeDataInterfaceDispatch.ush b/Engine/Plugins/Runtime/ComputeFramework/Shaders/Private/ComputeDataInterfaceDispatch.ush
new file mode 100644
--- /dev/null
+++ b/Engine/Plugins/Runtime/ComputeFramework/Shaders/Private/ComputeDataInterfaceDispatch.ush
@@ -0,0 +1,8 @@
+// Copyright Epic Games, Inc. All Rights Reserved.
+
+uint3 {DataInterfaceName}_NumThreads;
+
+uint3 ReadNumThreads_{DataInterfaceName}()
+{
+	return {DataInterfaceName}_NumThreads;
+}
diff --git a/Engine/Plugins/Runtime/ComputeFramework/Source/ComputeFramework/Private/ComputeDataInterfaceBuffer.cpp b/Engine/Plugins/Runtime/ComputeFramework/Source/ComputeFramework/Private/ComputeDataInterfaceBuffer.cpp
new file mode 100644
--- /dev/null
+++ b/Engine/Plugins/Runtime/ComputeFramework/Source/ComputeFramework/Private/ComputeDataInterfaceBuffer.cpp
@@ -0,0 +1,165 @@
+// Copyright Epic Games, Inc. All Rights Reserved.
+
+#include "ComputeFramework/ComputeDataInterfaceBuffer.h"
+
+#include "ComputeFramework/ShaderParamTypeDefinition.h"
+#include "RenderGraphBuilder.h"
+#include "RenderGraphUtils.h"
+#include "ShaderCompilerCore.h"
+#include "ShaderParameterMetadataBuilder.h"
+
+#include UE_INLINE_GENERATED_CPP_BY_NAME(ComputeDataInterfaceBuffer)
+
+static bool SupportsAtomics(FShaderValueTypeHandle InValueType)
+{
+	return (InValueType->DimensionType == EShaderFundamentalDimensionType::Scalar && (InValueType->Type == EShaderFundamentalType::Int || InValueType->Type == EShaderFundamentalType::Uint));
+}
+
+void UComputeDataInterfaceBuffer::GetSupportedInputs(TArray<FShaderFunctionDefinition>& OutFunctions) const
+{
+	OutFunctions.AddDefaulted_GetRef()
+		.SetName(TEXT("ReadNumValues"))
+		.AddReturnType(EShaderFundamentalType::Uint);
+
+	OutFunctions.AddDefaulted_GetRef()
+		.SetName(TEXT("ReadValue"))
+		.AddReturnType(ValueType)
+		.AddParam(EShaderFundamentalType::Uint);
+
+	OutFunctions.AddDefaulted_GetRef()
+		.SetName(TEXT("ReadValueUAV"))
+		.AddReturnType(ValueType)
+		.AddParam(EShaderFundamentalType::Uint);
+}
+
+void UComputeDataInterfaceBuffer::GetSupportedOutputs(TArray<FShaderFunctionDefinition>& OutFunctions) const
+{
+	OutFunctions.AddDefaulted_GetRef()
+		.SetName(TEXT("WriteValue"))
+		.AddParam(EShaderFundamentalType::Uint)
+		.AddParam(ValueType);
+
+	OutFunctions.AddDefaulted_GetRef()
+		.SetName(TEXT("WriteAtomicAdd"))
+		.AddReturnType(ValueType)
+		.AddParam(EShaderFundamentalType::Uint)
+		.AddParam(ValueType);
+
+	OutFunctions.AddDefaulted_GetRef()
+		.SetName(TEXT("WriteAtomicMin"))
+		.AddReturnType(ValueType)
+		.AddParam(EShaderFundamentalType::Uint)
+		.AddParam(ValueType);
+
+	OutFunctions.AddDefaulted_GetRef()
+		.SetName(TEXT("WriteAtomicMax"))
+		.AddReturnType(ValueType)
+		.AddParam(EShaderFundamentalType::Uint)
+		.AddParam(ValueType);
+}
+
+BEGIN_SHADER_PARAMETER_STRUCT(FBufferDataInterfaceParameters, )
+	SHADER_PARAMETER(uint32, BufferElementCount)
+	SHADER_PARAMETER_RDG_BUFFER_SRV(ByteAddressBuffer, BufferSRV)
+	SHADER_PARAMETER_RDG_BUFFER_UAV(RWByteAddressBuffer, BufferUAV)
+END_SHADER_PARAMETER_STRUCT()
+
+void UComputeDataInterfaceBuffer::GetShaderParameters(TCHAR const* UID, FShaderParametersMetadataBuilder& InOutBuilder, FShaderParametersMetadataAllocations& InOutAllocations) const
+{
+	InOutBuilder.AddNestedStruct<FBufferDataInterfaceParameters>(UID);
+}
+
+TCHAR const* UComputeDataInterfaceBuffer::TemplateFilePath = TEXT("/Plugin/ComputeFramework/Private/ComputeDataInterfaceBuffer.ush");
+
+TCHAR const* UComputeDataInterfaceBuffer::GetShaderVirtualPath() const
+{
+	return TemplateFilePath;
+}
+
+void UComputeDataInterfaceBuffer::GetShaderHash(FString& InOutKey) const
+{
+	GetShaderFileHash(TemplateFilePath, EShaderPlatform::SP_PCD3D_SM5).AppendString(InOutKey);
+}
+
+void UComputeDataInterfaceBuffer::GetHLSL(FString& OutHLSL, FString const& InDataInterfaceName) const
+{
+	const FString ValueTypeName = ValueType->ToString();
+	const int32 ValueTypeStride = ValueType->GetResourceElementSize();
+	const bool bSupportAtomic = SupportsAtomics(ValueType);
+
+	TMap<FString, FStringFormatArg> TemplateArgs =
+	{
+		{ TEXT("DataInterfaceName"), InDataInterfaceName },
+		{ TEXT("ValueType"), ValueTypeName },
+		{ TEXT("ValueTypeStride"), ValueTypeStride },
+		{ TEXT("SupportAtomic"), bSupportAtomic ? 1 : 0 },
+		{ TEXT("SplitReadWrite"), bAllowReadWrite ? 0 : 1 },
+	};
+
+	FString TemplateFile;
+	LoadShaderSourceFile(TemplateFilePath, EShaderPlatform::SP_PCD3D_SM5, &TemplateFile, nullptr);
+	OutHLSL += FString::Format(*TemplateFile, TemplateArgs);
+}
+
+UComputeDataProvider* UComputeDataInterfaceBuffer::CreateDataProvider() const
+{
+	UBufferDataProvider* Provider = NewObject<UBufferDataProvider>();
+	Provider->ElementCount = ElementCount;
+	Provider->ElementStride = ValueType->GetResourceElementSize();
+	Provider->bClearBeforeUse = bClearBeforeUse;
+	return Provider;
+}
+
+FComputeDataProviderRenderProxy* UBufferDataProvider::GetRenderProxy()
+{
+	return new FBufferDataProviderProxy(ElementCount, ElementStride, bClearBeforeUse);
+}
+
+FBufferDataProviderProxy::FBufferDataProviderProxy(int32 InElementCount, int32 InElementStride, bool bInClearBeforeUse)
+	: ElementCount(InElementCount)
+	, ElementStride(InElementStride)
+	, bClearBeforeUse(bInClearBeforeUse)
+{
+}
+
+bool FBufferDataProviderProxy::IsValid(FValidationData const& InValidationData) const
+{
+	if (InValidationData.ParameterStructSize != sizeof(FParameters))
+	{
+		return false;
+	}
+	
+	if (ElementCount <= 0)
+	{
+		return false;
+	}
+
+	return true;
+}
+
+void FBufferDataProviderProxy::AllocateResources(FRDGBuilder& GraphBuilder, FAllocationData const& InAllocationData)
+{
+	// If we are using a raw type alias for the buffer then we need to adjust stride and count.
+	check(ElementStride % 4 == 0);
+
+	Buffer = GraphBuilder.CreateBuffer(FRDGBufferDesc::CreateByteAddressDesc(ElementStride * ElementCount), TEXT("ComputeFrameworkBuffer"), ERDGBufferFlags::None);
+	BufferSRV = GraphBuilder.CreateSRV(Buffer);
+	BufferUAV = GraphBuilder.CreateUAV(Buffer, ERDGUnorderedAccessViewFlags::SkipBarrier);
+
+	if (bClearBeforeUse)
+	{
+		AddClearUAVPass(GraphBuilder, BufferUAV, 0);
+	}
+}
+
+void FBufferDataProviderProxy::GatherDispatchData(FDispatchData const& InDispatchData)
+{
+	TStridedView<FParameters> ParameterArray = MakeStridedParameterView<FParameters>(InDispatchData);
+	for (int32 InvocationIndex = 0; InvocationIndex < ParameterArray.Num(); ++InvocationIndex)
+	{
+		FParameters& Parameters = ParameterArray[InvocationIndex];
+		Parameters.BufferElementCount = ElementCount;
+		Parameters.BufferSRV = BufferSRV;
+		Parameters.BufferUAV = BufferUAV;
+	}
+}
diff --git a/Engine/Plugins/Runtime/ComputeFramework/Source/ComputeFramework/Private/ComputeDataInterfaceDispatch.cpp b/Engine/Plugins/Runtime/ComputeFramework/Source/ComputeFramework/Private/ComputeDataInterfaceDispatch.cpp
new file mode 100644
--- /dev/null
+++ b/Engine/Plugins/Runtime/ComputeFramework/Source/ComputeFramework/Private/ComputeDataInterfaceDispatch.cpp
@@ -0,0 +1,95 @@
+// Copyright Epic Games, Inc. All Rights Reserved.
+
+#include "ComputeFramework/ComputeDataInterfaceDispatch.h"
+
+#include "ComputeFramework/ShaderParamTypeDefinition.h"
+#include "RenderGraphBuilder.h"
+#include "ShaderCompilerCore.h"
+#include "ShaderParameterMetadataBuilder.h"
+
+#include UE_INLINE_GENERATED_CPP_BY_NAME(ComputeDataInterfaceDispatch)
+
+void UComputeDataInterfaceDispatch::GetSupportedInputs(TArray<FShaderFunctionDefinition>& OutFunctions) const
+{
+	OutFunctions.AddDefaulted_GetRef()
+		.SetName(TEXT("ReadNumThreads"))
+		.AddReturnType(FShaderValueType::Get(EShaderFundamentalType::Uint, 3));
+}
+
+BEGIN_SHADER_PARAMETER_STRUCT(FDataInterfaceDispatchParameters, )
+	SHADER_PARAMETER(FUintVector3, NumThreads)
+END_SHADER_PARAMETER_STRUCT()
+
+void UComputeDataInterfaceDispatch::GetShaderParameters(TCHAR const* UID, FShaderParametersMetadataBuilder& InOutBuilder, FShaderParametersMetadataAllocations& InOutAllocations) const
+{
+	InOutBuilder.AddNestedStruct<FDataInterfaceDispatchParameters>(UID);
+}
+
+const TCHAR* UComputeDataInterfaceDispatch::TemplateFilePath = TEXT("/Plugin/ComputeFramework/Private/ComputeDataInterfaceDispatch.ush");
+
+const TCHAR* UComputeDataInterfaceDispatch::GetShaderVirtualPath() const
+{
+	return TemplateFilePath;
+}
+
+void UComputeDataInterfaceDispatch::GetShaderHash(FString& InOutKey) const
+{
+	GetShaderFileHash(TemplateFilePath, EShaderPlatform::SP_PCD3D_SM5).AppendString(InOutKey);
+}
+
+void UComputeDataInterfaceDispatch::GetHLSL(FString& OutHLSL, FString const& InDataInterfaceName) const
+{
+	TMap<FString, FStringFormatArg> TemplateArgs =
+	{
+		{ TEXT("DataInterfaceName"), InDataInterfaceName },
+	};
+
+	FString TemplateFile;
+	LoadShaderSourceFile(TemplateFilePath, EShaderPlatform::SP_PCD3D_SM5, &TemplateFile, nullptr);
+	OutHLSL += FString::Format(*TemplateFile, TemplateArgs);
+}
+
+UComputeDataProvider* UComputeDataInterfaceDispatch::CreateDataProvider(TObjectPtr<UObject> InBinding,	uint64 InInputMask, uint64 InOutputMask) const
+{
+	UDispatchDataProvider* Provider = NewObject<UDispatchDataProvider>();
+	Provider->ThreadCount = ThreadCount;
+	return Provider;
+}
+
+FComputeDataProviderRenderProxy* UDispatchDataProvider::GetRenderProxy()
+{
+	FDispatchDataProviderProxy* Proxy = new FDispatchDataProviderProxy(ThreadCount);
+	return Proxy;
+}
+
+FDispatchDataProviderProxy::FDispatchDataProviderProxy(FUintVector InThreadCount)
+	: ThreadCount(InThreadCount)
+{
+}
+
+bool FDispatchDataProviderProxy::IsValid(FValidationData const& InValidationData) const
+{
+	if (InValidationData.ParameterStructSize != sizeof(FParameters))
+	{
+		return false;
+	}
+
+	return (ThreadCount.X > 0 && ThreadCount.Y > 0 && ThreadCount.Z > 0);
+}
+
+int32 FDispatchDataProviderProxy::GetDispatchThreadCount(TArray<FIntVector>& InOutThreadCounts) const
+{
+	InOutThreadCounts.Reset();
+	InOutThreadCounts.Add(FIntVector(ThreadCount.X, ThreadCount.Y, ThreadCount.Z));
+	return 1;
+}
+
+void FDispatchDataProviderProxy::GatherDispatchData(FDispatchData const& InDispatchData)
+{
+	const TStridedView<FParameters> ParameterArray = MakeStridedParameterView<FParameters>(InDispatchData);
+	for (int32 InvocationIndex = 0; InvocationIndex < ParameterArray.Num(); ++InvocationIndex)
+	{
+		FParameters& Parameters = ParameterArray[InvocationIndex];
+		Parameters.NumThreads = ThreadCount;
+	}
+}
diff --git a/Engine/Plugins/Runtime/ComputeFramework/Source/ComputeFramework/Private/ComputeFramework.cpp b/Engine/Plugins/Runtime/ComputeFramework/Source/ComputeFramework/Private/ComputeFramework.cpp
--- a/Engine/Plugins/Runtime/ComputeFramework/Source/ComputeFramework/Private/ComputeFramework.cpp
+++ b/Engine/Plugins/Runtime/ComputeFramework/Source/ComputeFramework/Private/ComputeFramework.cpp
@@ -5,7 +5,6 @@
 #include "ComputeFramework/ComputeFrameworkModule.h"
 #include "ComputeFramework/ComputeGraph.h"
 #include "ComputeFramework/ComputeGraphWorker.h"
-#include "ComputeFramework/ComputeKernelFromText.h"
 #include "ComputeFramework/ComputeKernelShaderCompilationManager.h"
 #include "ComputeFramework/ComputeSystem.h"
 #include "DataDrivenShaderPlatformInfo.h"
@@ -67,10 +66,6 @@ namespace ComputeFramework
 #if WITH_EDITOR
 		FlushShaderFileCache();
 
-		for (TObjectIterator<UComputeKernelFromText> It; It; ++It)
-		{
-			It->ReparseKernelSourceText();
-		}
 		for (TObjectIterator<UComputeGraph> It; It; ++It)
 		{
 			It->UpdateResources();
diff --git a/Engine/Plugins/Runtime/ComputeFramework/Source/ComputeFramework/Private/ComputeFramework/ComputeDataInterfaceBuffer.h b/Engine/Plugins/Runtime/ComputeFramework/Source/ComputeFramework/Private/ComputeFramework/ComputeDataInterfaceBuffer.h
new file mode 100644
--- /dev/null
+++ b/Engine/Plugins/Runtime/ComputeFramework/Source/ComputeFramework/Private/ComputeFramework/ComputeDataInterfaceBuffer.h
@@ -0,0 +1,89 @@
+// Copyright Epic Games, Inc. All Rights Reserved.
+
+#pragma once
+
+#include "ComputeFramework/ComputeDataInterface.h"
+#include "ComputeFramework/ComputeDataProvider.h"
+#include "ComputeFramework/ShaderParamTypeDefinition.h"
+
+#include "ComputeDataInterfaceBuffer.generated.h"
+
+class FBufferDataInterfaceParameters;
+
+/** Compute data interface used to own and give access to a GPU buffer. */
+UCLASS(MinimalAPI, Category = ComputeFramework)
+class UComputeDataInterfaceBuffer : public UComputeDataInterface
+{
+	GENERATED_BODY()
+
+	//~ Begin UComputeDataInterface Interface
+	TCHAR const* GetClassName() const override { return TEXT("Buffer"); }
+	void GetSupportedInputs(TArray<FShaderFunctionDefinition>& OutFunctions) const override;
+	void GetSupportedOutputs(TArray<FShaderFunctionDefinition>& OutFunctions) const override;
+	void GetShaderParameters(TCHAR const* UID, FShaderParametersMetadataBuilder& InOutBuilder, FShaderParametersMetadataAllocations& InOutAllocations) const override;
+	TCHAR const* GetShaderVirtualPath() const override;
+	void GetShaderHash(FString& InOutKey) const override;
+	void GetHLSL(FString& OutHLSL, FString const& InDataInterfaceName) const override;
+	UComputeDataProvider* CreateDataProvider() const override;
+	//~ End UComputeDataInterface Interface
+
+public:
+	/** The value type for the buffer. */
+	UPROPERTY()
+	FShaderValueTypeHandle ValueType;
+
+	/** The size of the buffer. */
+	UPROPERTY()
+	int32 ElementCount = 0;
+
+	/** Whether to allow read/write access. */
+	UPROPERTY()
+	bool bAllowReadWrite = false;
+
+	/** Whether to clear the buffer before use. */
+	UPROPERTY()
+	bool bClearBeforeUse = false;
+
+private:
+	static TCHAR const* TemplateFilePath;
+};
+
+/** Compute data provider implementation for UComputeDataInterfaceBuffer. */
+UCLASS(MinimalAPI)
+class UBufferDataProvider : public UComputeDataProvider
+{
+	GENERATED_BODY()
+
+	//~ Begin UComputeDataProvider Interface
+	FComputeDataProviderRenderProxy* GetRenderProxy() override;
+	//~ End UComputeDataProvider Interface
+
+public:
+	int32 ElementCount = 0;
+	int32 ElementStride = 4;
+	bool bClearBeforeUse;
+};
+
+/** Compute data provider proxy implementation for UComputeDataInterfaceBuffer. */
+class FBufferDataProviderProxy : public FComputeDataProviderRenderProxy
+{
+	//~ Begin FComputeDataProviderRenderProxy Interface
+	bool IsValid(FValidationData const& InValidationData) const override;
+	void AllocateResources(FRDGBuilder& GraphBuilder, FAllocationData const& InAllocationData) override;
+	void GatherDispatchData(FDispatchData const& InDispatchData) override;
+	//~ End FComputeDataProviderRenderProxy Interface
+
+public:
+	FBufferDataProviderProxy(int32 InElementCount, int32 InElementStride, bool bInClearBeforeUse);
+
+private:
+	using FParameters = FBufferDataInterfaceParameters;
+
+	const int32 ElementCount;
+	const int32 ElementStride;
+	const bool bClearBeforeUse;
+
+	FRDGBufferRef Buffer;
+	FRDGBufferSRVRef BufferSRV;
+	FRDGBufferUAVRef BufferUAV;
+};
diff --git a/Engine/Plugins/Runtime/ComputeFramework/Source/ComputeFramework/Private/ComputeFramework/ComputeDataInterfaceDispatch.h b/Engine/Plugins/Runtime/ComputeFramework/Source/ComputeFramework/Private/ComputeFramework/ComputeDataInterfaceDispatch.h
new file mode 100644
--- /dev/null
+++ b/Engine/Plugins/Runtime/ComputeFramework/Source/ComputeFramework/Private/ComputeFramework/ComputeDataInterfaceDispatch.h
@@ -0,0 +1,67 @@
+// Copyright Epic Games, Inc. All Rights Reserved.
+
+#pragma once
+
+#include "ComputeFramework/ComputeDataInterface.h"
+#include "ComputeFramework/ComputeDataProvider.h"
+
+#include "ComputeDataInterfaceDispatch.generated.h"
+
+class FDataInterfaceDispatchParameters;
+
+/** Compute data interface used to control a static dispatch size, and expose thread count to the shader. */
+UCLASS(MinimalAPI, Category = ComputeFramework)
+class UComputeDataInterfaceDispatch : public UComputeDataInterface
+{
+	GENERATED_BODY()
+
+	//~ Begin UComputeDataInterface Interface
+	TCHAR const* GetClassName() const override { return TEXT("Dispatch"); }
+	bool IsExecutionInterface() const override { return true; }
+	void GetSupportedInputs(TArray<FShaderFunctionDefinition>& OutFunctions) const override;
+	void GetShaderParameters(TCHAR const* UID, FShaderParametersMetadataBuilder& InOutBuilder, FShaderParametersMetadataAllocations& InOutAllocations) const override;
+	const TCHAR* GetShaderVirtualPath() const override;
+	void GetShaderHash(FString& InOutKey) const override;
+	void GetHLSL(FString& OutHLSL, FString const& InDataInterfaceName) const override;
+	UComputeDataProvider* CreateDataProvider(TObjectPtr<UObject> InBinding, uint64 InInputMask, uint64 InOutputMask) const override;
+	//~ End UComputeDataInterface Interface
+
+public:
+	UPROPERTY()
+	FUintVector ThreadCount;
+	
+private:
+	static const TCHAR* TemplateFilePath;
+};
+
+/** Compute data provider implementation for UComputeDataInterfaceDispatch. */
+UCLASS()
+class UDispatchDataProvider : public UComputeDataProvider
+{
+	GENERATED_BODY()
+
+	//~ Begin UComputeDataProvider Interface
+	FComputeDataProviderRenderProxy* GetRenderProxy() override;
+	//~ End UComputeDataProvider Interface
+
+public:
+	FUintVector ThreadCount;
+};
+
+/** Compute data provider proxy implementation for UComputeDataInterfaceDispatch. */
+class FDispatchDataProviderProxy : public FComputeDataProviderRenderProxy
+{
+	//~ Begin FComputeDataProviderRenderProxy Interface
+	bool IsValid(FValidationData const& InValidationData) const override;
+	int32 GetDispatchThreadCount(TArray<FIntVector>& InOutThreadCounts) const override;
+	void GatherDispatchData(FDispatchData const& InDispatchData) override;
+	//~ End FComputeDataProviderRenderProxy Interface
+
+public:
+	FDispatchDataProviderProxy(FUintVector InThreadCount);
+
+private:
+	using FParameters = FDataInterfaceDispatchParameters;
+	
+	FUintVector ThreadCount;
+};
diff --git a/Engine/Plugins/Runtime/ComputeFramework/Source/ComputeFramework/Private/ComputeGraphComponent.cpp b/Engine/Plugins/Runtime/ComputeFramework/Source/ComputeFramework/Private/ComputeGraphComponent.cpp
--- a/Engine/Plugins/Runtime/ComputeFramework/Source/ComputeFramework/Private/ComputeGraphComponent.cpp
+++ b/Engine/Plugins/Runtime/ComputeFramework/Source/ComputeFramework/Private/ComputeGraphComponent.cpp
@@ -1,8 +1,9 @@
 // Copyright Epic Games, Inc. All Rights Reserved.
 
 #include "ComputeFramework/ComputeGraphComponent.h"
+
 #include "ComputeFramework/ComputeFramework.h"
-#include "ComputeFramework/ComputeGraph.h"
+#include "ComputeFramework/ComputeGraphInstance.h"
 #include "ComputeWorkerInterface.h"
 #include "GameFramework/Actor.h"
 
@@ -17,8 +18,16 @@ UComputeGraphComponent::UComputeGraphComponent()
 	PrimaryComponentTick.bStartWithTickEnabled = false;
 }
 
-UComputeGraphComponent::~UComputeGraphComponent()
+void UComputeGraphComponent::OnRegister()
 {
+	Super::OnRegister();
+	CreateDataProviders(0, this);
+}
+
+void UComputeGraphComponent::OnUnregister()
+{
+	DestroyDataProviders();
+	Super::OnUnregister();
 }
 
 void UComputeGraphComponent::CreateDataProviders(int32 InBindingIndex, UObject* InBindingObject)
@@ -39,7 +48,7 @@ void UComputeGraphComponent::QueueExecute()
 	}
 }
 
-void UComputeGraphComponent::TickComponent(float DeltaTime, enum ELevelTick TickType, FActorComponentTickFunction* ThisTickFunction)
+void UComputeGraphComponent::TickComponent(float DeltaTime, ELevelTick TickType, FActorComponentTickFunction* ThisTickFunction)
 {
 	Super::TickComponent(DeltaTime, TickType, ThisTickFunction);
 
@@ -58,4 +67,4 @@ void UComputeGraphComponent::DestroyRenderState_Concurrent()
 	Super::DestroyRenderState_Concurrent();
 	
 	ComputeFramework::AbortWork(GetScene(), this);
-}
\ No newline at end of file
+}
diff --git a/Engine/Plugins/Runtime/ComputeFramework/Source/ComputeFramework/Private/ComputeGraphFromText.cpp b/Engine/Plugins/Runtime/ComputeFramework/Source/ComputeFramework/Private/ComputeGraphFromText.cpp
new file mode 100644
--- /dev/null
+++ b/Engine/Plugins/Runtime/ComputeFramework/Source/ComputeFramework/Private/ComputeGraphFromText.cpp
@@ -0,0 +1,183 @@
+// Copyright Epic Games, Inc. All Rights Reserved.
+
+#include "ComputeFramework/ComputeGraphFromText.h"
+
+#include "Algo/Count.h"
+#include "ComputeFramework/ComputeDataInterfaceBuffer.h"
+#include "ComputeFramework/ComputeDataInterfaceDispatch.h"
+#include "ComputeFramework/ComputeGraphComponent.h"
+#include "ComputeFramework/ComputeKernel.h"
+#include "ComputeFramework/ComputeKernelSource.h"
+
+#include UE_INLINE_GENERATED_CPP_BY_NAME(ComputeGraphFromText)
+
+#if WITH_EDITOR
+
+void UComputeGraphFromText::SetGraphSourceText(FString const& InText)
+{
+	GraphSourceText = InText;
+	Rebuild();
+}
+
+void UComputeGraphFromText::PostLoad()
+{
+	Rebuild();
+	Super::PostLoad();
+}
+
+void UComputeGraphFromText::PostEditChangeProperty(FPropertyChangedEvent& PropertyChangedEvent)
+{
+	Super::PostEditChangeProperty(PropertyChangedEvent);
+
+	if (PropertyChangedEvent.MemberProperty && PropertyChangedEvent.Property)
+	{
+		if (PropertyChangedEvent.MemberProperty->GetFName() == GET_MEMBER_NAME_CHECKED(UComputeGraphFromText, GraphSourceText))
+		{
+			Rebuild();
+			UpdateResources();
+		}
+	}
+}
+
+void UComputeGraphFromText::Reset()
+{
+	KernelInvocations.Reset();
+	DataInterfaces.Reset();
+	GraphEdges.Reset();
+	Bindings.Reset();
+	DataInterfaceToBinding.Reset();
+}
+
+int32 UComputeGraphFromText::AddKernel(FString const& InSourceText, FString const& InEntryPoint, FIntVector const& InGroupSize)
+{
+	UComputeKernel* Kernel = NewObject<UComputeKernel>(this);
+	const int32 KernelIndex = KernelInvocations.Add(Kernel);
+
+	UComputeKernelSourceWithText* KernelSource = NewObject<UComputeKernelSourceWithText>(Kernel);
+	Kernel->KernelSource = KernelSource;
+	
+	KernelSource->SourceText = InSourceText;
+	KernelSource->EntryPoint = InEntryPoint;
+	KernelSource->GroupSize = InGroupSize;
+
+	return KernelIndex;
+}
+
+int32 UComputeGraphFromText::AddObjectBinding(UClass* InClass)
+{
+	const int32 BindingIndex = Bindings.Add(InClass);
+	return BindingIndex;
+}
+
+int32 UComputeGraphFromText::AddDataInterface(UComputeDataInterface* InDataInterface, int32 InBindingIndex)
+{
+	const int32 DataInterfaceIndex = DataInterfaces.Add(InDataInterface);
+	DataInterfaceToBinding.Add(InBindingIndex);
+	return DataInterfaceIndex;
+}
+
+bool UComputeGraphFromText::AddInputGraphConnection(int32 InKernelIndex, FString const& InKernelFunctionName, int32 InDataProviderIndex, FString const& InDataProviderFunctionName)
+{
+	if (!KernelInvocations.IsValidIndex(InKernelIndex) || !DataInterfaces.IsValidIndex(InDataProviderIndex))
+	{
+		return false;
+	}
+	
+	// Find matching function exposed by data provider.
+	TArray<FShaderFunctionDefinition> Functions;
+	DataInterfaces[InDataProviderIndex]->GetSupportedInputs(Functions);
+	const int32 FunctionIndex = Functions.IndexOfByPredicate([InDataProviderFunctionName](FShaderFunctionDefinition const& Def) { return Def.Name == InDataProviderFunctionName; });
+	if (FunctionIndex == INDEX_NONE)
+	{
+		return false;
+	}
+	
+	// Store input function and rename.
+	KernelInvocations[InKernelIndex]->KernelSource->ExternalInputs.Add_GetRef(Functions[FunctionIndex]).SetName(InKernelFunctionName);
+
+	// Count current kernel binding index.
+	const int32 KernelInputIndex = Algo::CountIf(GraphEdges, [InKernelIndex](FComputeGraphEdge const& Edge) { return Edge.KernelIndex == InKernelIndex && Edge.bKernelInput; });
+
+	FComputeGraphEdge& Edge = GraphEdges.AddDefaulted_GetRef();
+	Edge.KernelIndex = InKernelIndex;
+	Edge.KernelBindingIndex = KernelInputIndex;
+	Edge.DataInterfaceIndex = InDataProviderIndex;
+	Edge.DataInterfaceBindingIndex = FunctionIndex;
+	Edge.bKernelInput = true;
+
+	return true;
+}
+
+bool UComputeGraphFromText::AddOutputGraphConnection(int32 InKernelIndex, FString const& InKernelFunctionName, int32 InDataProviderIndex, FString const& InDataProviderFunctionName)
+{
+	if (!KernelInvocations.IsValidIndex(InKernelIndex) || !DataInterfaces.IsValidIndex(InDataProviderIndex))
+	{
+		return false;
+	}
+
+	// Find matching function exposed by data provider.
+	TArray<FShaderFunctionDefinition> Functions;
+	DataInterfaces[InDataProviderIndex]->GetSupportedOutputs(Functions);
+	
+	const int32 FunctionIndex = Functions.IndexOfByPredicate([InDataProviderFunctionName](FShaderFunctionDefinition const& Def) { return Def.Name == InDataProviderFunctionName; });
+	if (FunctionIndex == INDEX_NONE)
+	{
+		return false;
+	}
+
+	// Store input function and rename.
+	KernelInvocations[InKernelIndex]->KernelSource->ExternalOutputs.Add_GetRef(Functions[FunctionIndex]).SetName(InKernelFunctionName);
+
+	// Count current kernel binding index.
+	const int32 KernelOutputIndex = Algo::CountIf(GraphEdges, [InKernelIndex](FComputeGraphEdge const& Edge) { return Edge.KernelIndex == InKernelIndex && !Edge.bKernelInput; });
+
+	FComputeGraphEdge& Edge = GraphEdges.AddDefaulted_GetRef();
+	Edge.KernelIndex = InKernelIndex;
+	Edge.KernelBindingIndex = KernelOutputIndex;
+	Edge.DataInterfaceIndex = InDataProviderIndex;
+	Edge.DataInterfaceBindingIndex = FunctionIndex;
+	Edge.bKernelInput = false;
+
+	return true;
+}
+
+void UComputeGraphFromText::Rebuild()
+{
+	Reset();
+
+	// Create a simple copy buffer kernel.
+	FString SourceText(TEXT("[numthreads(64, 1, 1)]\nvoid CopyCS(uint DispatchThreadId : SV_DispatchThreadID)\n{\n\tuint ThreadIndex = DispatchThreadId.x;\n\tif (ThreadIndex >= ReadNumValues())\n\t{\n\t\treturn;\n\t}\n\tWriteValue(ThreadIndex, ReadValue(ThreadIndex));\n}"));
+	FString EntryPoint(TEXT("CopyCS"));
+	FIntVector GroupSize(64, 1, 1);
+	const int32 CopyBufferKernelIndex = AddKernel(SourceText, EntryPoint, GroupSize);
+
+	// Set some default (but unused here) object binding type.
+	const int32 DefaultBindingClassIndex = AddObjectBinding(UComputeGraphComponent::StaticClass());
+
+	// Create dispatch data interface.
+	UComputeDataInterfaceDispatch* DispatchDataInterface = NewObject<UComputeDataInterfaceDispatch>(this);
+	DispatchDataInterface->ThreadCount = FUintVector3(100, 1, 1);
+	const int32 DispatchDataInterfaceIndex = AddDataInterface(DispatchDataInterface, DefaultBindingClassIndex);
+
+	// Create source buffer data interface.
+	UComputeDataInterfaceBuffer* SourceBufferDataInterface = NewObject<UComputeDataInterfaceBuffer>(this);
+	SourceBufferDataInterface->ValueType = FShaderValueType::Get(EShaderFundamentalType::Int, 3);
+	SourceBufferDataInterface->ElementCount = 100;
+	SourceBufferDataInterface->bClearBeforeUse = true;
+	const int32 SourceBufferDataInterfaceIndex = AddDataInterface(SourceBufferDataInterface, DefaultBindingClassIndex);
+
+	// Create destination buffer data interface.
+	UComputeDataInterfaceBuffer* DestBufferDataInterface = NewObject<UComputeDataInterfaceBuffer>(this);
+	DestBufferDataInterface->ValueType = FShaderValueType::Get(EShaderFundamentalType::Int, 3);
+	DestBufferDataInterface->ElementCount = 100;
+	const int32 DestBufferDataInterfaceIndex = AddDataInterface(DestBufferDataInterface, DefaultBindingClassIndex);;
+
+	// Bind kernel functions (as used in the kernel source) to the data interface functions (as defined in the data interface).
+	// Note that the DispatchDataInterface is special and we must bind one of its functions even if we don't actually use any.
+	AddInputGraphConnection(CopyBufferKernelIndex, TEXT("ReadNumThreads"), DispatchDataInterfaceIndex, TEXT("ReadNumThreads"));
+	AddInputGraphConnection(CopyBufferKernelIndex, TEXT("ReadNumValues"), SourceBufferDataInterfaceIndex, TEXT("ReadNumValues"));
+	AddInputGraphConnection(CopyBufferKernelIndex, TEXT("ReadValue"), SourceBufferDataInterfaceIndex, TEXT("ReadValue"));
+	AddOutputGraphConnection(CopyBufferKernelIndex, TEXT("WriteValue"), DestBufferDataInterfaceIndex, TEXT("WriteValue"));
+}
+
+#endif
diff --git a/Engine/Plugins/Runtime/ComputeFramework/Source/ComputeFramework/Private/ComputeKernelFromText.cpp b/Engine/Plugins/Runtime/ComputeFramework/Source/ComputeFramework/Private/ComputeKernelFromText.cpp
deleted file mode 100644
diff --git a/Engine/Plugins/Runtime/ComputeFramework/Source/ComputeFramework/Private/ComputeSourceFromText.cpp b/Engine/Plugins/Runtime/ComputeFramework/Source/ComputeFramework/Private/ComputeSourceFromText.cpp
deleted file mode 100644
diff --git a/Engine/Plugins/Runtime/ComputeFramework/Source/ComputeFramework/Private/ShaderParamTypeDefinition.cpp b/Engine/Plugins/Runtime/ComputeFramework/Source/ComputeFramework/Private/ShaderParamTypeDefinition.cpp
--- a/Engine/Plugins/Runtime/ComputeFramework/Source/ComputeFramework/Private/ShaderParamTypeDefinition.cpp
+++ b/Engine/Plugins/Runtime/ComputeFramework/Source/ComputeFramework/Private/ShaderParamTypeDefinition.cpp
@@ -767,6 +767,7 @@ FShaderFunctionDefinition& FShaderFunctionDefinition::AddParam(FShaderValueTypeH
 	FShaderParamTypeDefinition Def;
 	Def.ValueType = InValueType;
 	Def.Modifier = InModifier;
+	Def.ResetTypeDeclaration();
 	ParamTypes.Add(Def);
 	return *this;
 }
@@ -789,6 +790,8 @@ FShaderFunctionDefinition& FShaderFunctionDefinition::AddParam(EShaderFundamenta
 		Def.ValueType = FShaderValueType::Get(InType);
 	}
 
+	Def.ResetTypeDeclaration();
+
 	ParamTypes.Add(Def);
 	return *this;
 }
diff --git a/Engine/Plugins/Runtime/ComputeFramework/Source/ComputeFramework/Public/ComputeFramework/ComputeGraphComponent.h b/Engine/Plugins/Runtime/ComputeFramework/Source/ComputeFramework/Public/ComputeFramework/ComputeGraphComponent.h
--- a/Engine/Plugins/Runtime/ComputeFramework/Source/ComputeFramework/Public/ComputeFramework/ComputeGraphComponent.h
+++ b/Engine/Plugins/Runtime/ComputeFramework/Source/ComputeFramework/Public/ComputeFramework/ComputeGraphComponent.h
@@ -21,7 +21,6 @@ class UComputeGraphComponent : public UActorComponent
 
 public:
 	UE_API UComputeGraphComponent();
-	UE_API ~UComputeGraphComponent();
 
 	/** The Compute Graph asset. */
 	UPROPERTY(EditAnywhere, BlueprintReadWrite, Category = "Compute")
@@ -41,9 +40,11 @@ public:
 
 protected:
 	//~ Begin UActorComponent Interface
-	UE_API void TickComponent(float DeltaTime, enum ELevelTick TickType, FActorComponentTickFunction* ThisTickFunction) override;
-	UE_API void SendRenderDynamicData_Concurrent() override;
-	UE_API void DestroyRenderState_Concurrent() override;
+	void OnRegister() override;
+	void OnUnregister() override;
+	void TickComponent(float DeltaTime, enum ELevelTick TickType, FActorComponentTickFunction* ThisTickFunction) override;
+	void SendRenderDynamicData_Concurrent() override;
+	void DestroyRenderState_Concurrent() override;
 	bool ShouldCreateRenderState() const override { return true; }
 	//~ End UActorComponent Interface
 
diff --git a/Engine/Plugins/Runtime/ComputeFramework/Source/ComputeFramework/Public/ComputeFramework/ComputeGraphFromText.h b/Engine/Plugins/Runtime/ComputeFramework/Source/ComputeFramework/Public/ComputeFramework/ComputeGraphFromText.h
new file mode 100644
--- /dev/null
+++ b/Engine/Plugins/Runtime/ComputeFramework/Source/ComputeFramework/Public/ComputeFramework/ComputeGraphFromText.h
@@ -0,0 +1,43 @@
+// Copyright Epic Games, Inc. All Rights Reserved.
+
+#pragma once
+
+#include "ComputeGraph.h"
+#include "ComputeGraphFromText.generated.h"
+
+#define UE_API COMPUTEFRAMEWORK_API
+
+/** Simple data driven compute graph implementation. */
+UCLASS(MinimalAPI)
+class UComputeGraphFromText : public UComputeGraph
+{
+	GENERATED_BODY()
+
+#if WITH_EDITORONLY_DATA
+	/** Source text containing the full graph description. */
+	UPROPERTY(EditAnywhere, Category = "Graph", meta = (MultiLine = true))
+	FString GraphSourceText;
+#endif
+
+#if WITH_EDITOR
+	//~ Begin UObject Interface.
+	UE_API void PostLoad() override;
+	UE_API void PostEditChangeProperty(FPropertyChangedEvent& PropertyChangedEvent) override;
+	//~ End UObject Interface.
+
+public:
+	/** Set the source text that describes the graph. */
+	void SetGraphSourceText(FString const& InText);
+
+private:
+	void Reset();
+	int32 AddKernel(FString const& InSourceText, FString const& InEntryPoint, FIntVector const& InGroupSize);
+	int32 AddObjectBinding(UClass* InClass);
+	int32 AddDataInterface(UComputeDataInterface* InDataInterface, int32 InBindingIndex);
+	bool AddInputGraphConnection(int32 InKernelIndex, FString const& InKernelFunctionName, int32 InDataProviderIndex, FString const& InDataProviderFunctionName);
+	bool AddOutputGraphConnection(int32 InKernelIndex, FString const& InKernelFunctionName, int32 InDataProviderIndex, FString const& InDataProviderFunctionName);
+	void Rebuild();
+#endif
+};
+
+#undef UE_API
diff --git a/Engine/Plugins/Runtime/ComputeFramework/Source/ComputeFramework/Public/ComputeFramework/ComputeKernelFromText.h b/Engine/Plugins/Runtime/ComputeFramework/Source/ComputeFramework/Public/ComputeFramework/ComputeKernelFromText.h
deleted file mode 100644
diff --git a/Engine/Plugins/Runtime/ComputeFramework/Source/ComputeFramework/Public/ComputeFramework/ComputeKernelSource.h b/Engine/Plugins/Runtime/ComputeFramework/Source/ComputeFramework/Public/ComputeFramework/ComputeKernelSource.h
--- a/Engine/Plugins/Runtime/ComputeFramework/Source/ComputeFramework/Public/ComputeFramework/ComputeKernelSource.h
+++ b/Engine/Plugins/Runtime/ComputeFramework/Source/ComputeFramework/Public/ComputeFramework/ComputeKernelSource.h
@@ -14,7 +14,7 @@ class UComputeSource;
  * Class representing the source for a UComputeKernel 
  * We derive from this for each authoring mechanism. (HLSL text, VPL graph, ML Meta Lang, etc.)
  */
-UCLASS(MinimalAPI, Abstract, EditInlineNew)
+UCLASS(MinimalAPI, Abstract)
 class UComputeKernelSource : public UObject
 {
 	GENERATED_BODY()
@@ -23,7 +23,6 @@ public:
 	/** Get kernel source code ready for HLSL compilation. */
 	virtual FString GetSource() const PURE_VIRTUAL(UComputeKernelSource::GetSource, return {};);
 
-public:
 	/** Kernel entry point. */
 	UPROPERTY(VisibleAnywhere, AssetRegistrySearchable, Category = "Kernel")
 	FString EntryPoint;
@@ -53,4 +52,20 @@ public:
 	TArray<FShaderFunctionDefinition> ExternalOutputs;
 };
 
+/** Simple concrete implementation of UComputeKernelSource with contained source text property. */
+UCLASS(MinimalAPI)
+class UComputeKernelSourceWithText final : public UComputeKernelSource
+{
+	GENERATED_BODY()
+
+public:
+	/* Kernel source text. */
+	UPROPERTY(EditAnywhere, Category = "Kernel")
+	FString SourceText;
+
+	//~ Begin UComputeKernelSource Interface.
+	virtual FString GetSource() const override { return SourceText; }
+	//~ End UComputeKernelSource Interface.
+};
+
 #undef UE_API
diff --git a/Engine/Plugins/Runtime/ComputeFramework/Source/ComputeFramework/Public/ComputeFramework/ComputeSourceFromText.h b/Engine/Plugins/Runtime/ComputeFramework/Source/ComputeFramework/Public/ComputeFramework/ComputeSourceFromText.h
deleted file mode 100644
diff --git a/Engine/Plugins/Runtime/ComputeFramework/Source/ComputeFrameworkEditor/Private/ComputeFramework/ComputeGraphFromTextAssetActions.h b/Engine/Plugins/Runtime/ComputeFramework/Source/ComputeFrameworkEditor/Private/ComputeFramework/ComputeGraphFromTextAssetActions.h
new file mode 100644
--- /dev/null
+++ b/Engine/Plugins/Runtime/ComputeFramework/Source/ComputeFrameworkEditor/Private/ComputeFramework/ComputeGraphFromTextAssetActions.h
@@ -0,0 +1,22 @@
+// Copyright Epic Games, Inc. All Rights Reserved.
+
+#pragma once
+
+#include "AssetTypeCategories.h"
+#include "AssetTypeActions_Base.h"
+
+class FAssetTypeActions_ComputeGraphFromText : public FAssetTypeActions_Base
+{
+	//~ Begin FAssetTypeActions_Base Interface.
+	FText GetName() const override;
+	FColor GetTypeColor() const override;
+	UClass* GetSupportedClass() const override;
+	uint32 GetCategories() override;
+	//~ End FAssetTypeActions_Base Interface.
+
+public:
+	FAssetTypeActions_ComputeGraphFromText(EAssetTypeCategories::Type InAssetCategoryBit = EAssetTypeCategories::Misc);
+
+private:
+	EAssetTypeCategories::Type AssetCategoryBit;
+};
diff --git a/Engine/Plugins/Runtime/ComputeFramework/Source/ComputeFrameworkEditor/Private/ComputeFramework/ComputeGraphFromTextFactory.h b/Engine/Plugins/Runtime/ComputeFramework/Source/ComputeFrameworkEditor/Private/ComputeFramework/ComputeGraphFromTextFactory.h
new file mode 100644
--- /dev/null
+++ b/Engine/Plugins/Runtime/ComputeFramework/Source/ComputeFrameworkEditor/Private/ComputeFramework/ComputeGraphFromTextFactory.h
@@ -0,0 +1,26 @@
+// Copyright Epic Games, Inc. All Rights Reserved.
+
+#pragma once
+
+#include "Factories/Factory.h"
+#include "ComputeGraphFromTextFactory.generated.h"
+
+UCLASS(hidecategories = Object)
+class UComputeGraphFromTextFactory : public UFactory
+{
+	GENERATED_BODY()
+
+	//~ Begin UFactory Interface.
+	virtual UObject* FactoryCreateNew(
+		UClass* InClass, 
+		UObject* InParent, 
+		FName InName, 
+		EObjectFlags Flags, 
+		UObject* Context, 
+		FFeedbackContext* Warn
+		) override;
+	//~ End UFactory Interface.
+
+public:
+	UComputeGraphFromTextFactory();
+};
diff --git a/Engine/Plugins/Runtime/ComputeFramework/Source/ComputeFrameworkEditor/Private/ComputeFramework/ComputeKernelAssetActions.h b/Engine/Plugins/Runtime/ComputeFramework/Source/ComputeFrameworkEditor/Private/ComputeFramework/ComputeKernelAssetActions.h
deleted file mode 100644
diff --git a/Engine/Plugins/Runtime/ComputeFramework/Source/ComputeFrameworkEditor/Private/ComputeFramework/ComputeKernelFactory.h b/Engine/Plugins/Runtime/ComputeFramework/Source/ComputeFrameworkEditor/Private/ComputeFramework/ComputeKernelFactory.h
deleted file mode 100644
diff --git a/Engine/Plugins/Runtime/ComputeFramework/Source/ComputeFrameworkEditor/Private/ComputeFramework/ComputeKernelFromTextAssetActions.h b/Engine/Plugins/Runtime/ComputeFramework/Source/ComputeFrameworkEditor/Private/ComputeFramework/ComputeKernelFromTextAssetActions.h
deleted file mode 100644
diff --git a/Engine/Plugins/Runtime/ComputeFramework/Source/ComputeFrameworkEditor/Private/ComputeFramework/ComputeKernelFromTextFactory.h b/Engine/Plugins/Runtime/ComputeFramework/Source/ComputeFrameworkEditor/Private/ComputeFramework/ComputeKernelFromTextFactory.h
deleted file mode 100644
diff --git a/Engine/Plugins/Runtime/ComputeFramework/Source/ComputeFrameworkEditor/Private/ComputeFrameworkEditorModule.cpp b/Engine/Plugins/Runtime/ComputeFramework/Source/ComputeFrameworkEditor/Private/ComputeFrameworkEditorModule.cpp
--- a/Engine/Plugins/Runtime/ComputeFramework/Source/ComputeFrameworkEditor/Private/ComputeFrameworkEditorModule.cpp
+++ b/Engine/Plugins/Runtime/ComputeFramework/Source/ComputeFrameworkEditor/Private/ComputeFrameworkEditorModule.cpp
@@ -1,20 +1,18 @@
 // Copyright Epic Games, Inc. All Rights Reserved.
 
 #include "ComputeFramework/ComputeFrameworkEditorModule.h"
-#include "ComputeFramework/ComputeFrameworkCompilationTick.h"
 
-#include "ComputeFramework/ComputeKernelAssetActions.h"
-#include "ComputeFramework/ComputeKernelFromTextAssetActions.h"
-#include "IAssetTools.h"
+#include "AssetToolsModule.h"
+#include "ComputeFramework/ComputeFrameworkCompilationTick.h"
+#include "ComputeFramework/ComputeGraphFromTextAssetActions.h"
 #include "Modules/ModuleManager.h"
 
 void FComputeFrameworkEditorModule::StartupModule()
 {
 	TickObject = MakeUnique<FComputeFrameworkCompilationTick>();
 
-// 	IAssetTools& AssetTools = FModuleManager::LoadModuleChecked<FAssetToolsModule>("AssetTools").Get();
-// 	AssetTools.RegisterAssetTypeActions(MakeShareable(new FAssetTypeActions_ComputeKernel));
-// 	AssetTools.RegisterAssetTypeActions(MakeShareable(new FAssetTypeActions_ComputeKernelFromText));
+ 	IAssetTools& AssetTools = FModuleManager::LoadModuleChecked<FAssetToolsModule>("AssetTools").Get();
+ 	AssetTools.RegisterAssetTypeActions(MakeShareable(new FAssetTypeActions_ComputeGraphFromText));
 }
 
 void FComputeFrameworkEditorModule::ShutdownModule()
diff --git a/Engine/Plugins/Runtime/ComputeFramework/Source/ComputeFrameworkEditor/Private/ComputeGraphFromTextAssetActions.cpp b/Engine/Plugins/Runtime/ComputeFramework/Source/ComputeFrameworkEditor/Private/ComputeGraphFromTextAssetActions.cpp
new file mode 100644
--- /dev/null
+++ b/Engine/Plugins/Runtime/ComputeFramework/Source/ComputeFrameworkEditor/Private/ComputeGraphFromTextAssetActions.cpp
@@ -0,0 +1,30 @@
+// Copyright Epic Games, Inc. All Rights Reserved.
+
+#include "ComputeFramework/ComputeGraphFromTextAssetActions.h"
+
+#include "ComputeFramework/ComputeGraphFromText.h"
+
+FAssetTypeActions_ComputeGraphFromText::FAssetTypeActions_ComputeGraphFromText(EAssetTypeCategories::Type InAssetCategoryBit)
+	: AssetCategoryBit(InAssetCategoryBit)
+{
+}
+
+FText FAssetTypeActions_ComputeGraphFromText::GetName() const
+{
+	return NSLOCTEXT("ComputeFramework", "ComputeGraphFromTextName", "Compute Graph (Text)");
+}
+
+FColor FAssetTypeActions_ComputeGraphFromText::GetTypeColor() const
+{
+	return FColor::Turquoise;
+}
+
+UClass* FAssetTypeActions_ComputeGraphFromText::GetSupportedClass() const
+{
+	return UComputeGraphFromText::StaticClass();
+}
+
+uint32 FAssetTypeActions_ComputeGraphFromText::GetCategories()
+{
+	return AssetCategoryBit;
+}
diff --git a/Engine/Plugins/Runtime/ComputeFramework/Source/ComputeFrameworkEditor/Private/ComputeGraphFromTextFactory.cpp b/Engine/Plugins/Runtime/ComputeFramework/Source/ComputeFrameworkEditor/Private/ComputeGraphFromTextFactory.cpp
new file mode 100644
--- /dev/null
+++ b/Engine/Plugins/Runtime/ComputeFramework/Source/ComputeFrameworkEditor/Private/ComputeGraphFromTextFactory.cpp
@@ -0,0 +1,19 @@
+// Copyright Epic Games, Inc. All Rights Reserved.
+
+#include "ComputeFramework/ComputeGraphFromTextFactory.h"
+
+#include "ComputeFramework/ComputeGraphFromText.h"
+
+UComputeGraphFromTextFactory::UComputeGraphFromTextFactory()
+{
+	SupportedClass = UComputeGraphFromText::StaticClass();
+	bCreateNew = true;
+	bEditAfterNew = true;
+}
+
+UObject* UComputeGraphFromTextFactory::FactoryCreateNew(UClass* InClass, UObject* InParent, FName InName, EObjectFlags Flags, UObject* Context, FFeedbackContext* Warn)
+{
+	UComputeGraphFromText* Kernel = NewObject<UComputeGraphFromText>(InParent, InClass, InName, Flags);
+
+	return Kernel;
+}
diff --git a/Engine/Plugins/Runtime/ComputeFramework/Source/ComputeFrameworkEditor/Private/ComputeKernelAssetActions.cpp b/Engine/Plugins/Runtime/ComputeFramework/Source/ComputeFrameworkEditor/Private/ComputeKernelAssetActions.cpp
deleted file mode 100644
diff --git a/Engine/Plugins/Runtime/ComputeFramework/Source/ComputeFrameworkEditor/Private/ComputeKernelFactory.cpp b/Engine/Plugins/Runtime/ComputeFramework/Source/ComputeFrameworkEditor/Private/ComputeKernelFactory.cpp
deleted file mode 100644
diff --git a/Engine/Plugins/Runtime/ComputeFramework/Source/ComputeFrameworkEditor/Private/ComputeKernelFromTextAssetActions.cpp b/Engine/Plugins/Runtime/ComputeFramework/Source/ComputeFrameworkEditor/Private/ComputeKernelFromTextAssetActions.cpp
deleted file mode 100644
diff --git a/Engine/Plugins/Runtime/ComputeFramework/Source/ComputeFrameworkEditor/Private/ComputeKernelFromTextFactory.cpp b/Engine/Plugins/Runtime/ComputeFramework/Source/ComputeFrameworkEditor/Private/ComputeKernelFromTextFactory.cpp
deleted file mode 100644
