ComputeFramework: Fix typo in data provider initialize.
#rnx

diff --git a/Engine/Plugins/Runtime/ComputeFramework/Source/ComputeFramework/Private/ComputeGraph.cpp b/Engine/Plugins/Runtime/ComputeFramework/Source/ComputeFramework/Private/ComputeGraph.cpp
--- a/Engine/Plugins/Runtime/ComputeFramework/Source/ComputeFramework/Private/ComputeGraph.cpp
+++ b/Engine/Plugins/Runtime/ComputeFramework/Source/ComputeFramework/Private/ComputeGraph.cpp
@@ -286,7 +286,7 @@ void UComputeGraph::CreateDataProviders(int32 InBindingIndex, TObjectPtr<UObject
 
 					if (DataProvider)
 					{
-						DataProvider->Initialize(DataInterface, InBindingObject, InputMask, InputMask);
+						DataProvider->Initialize(DataInterface, InBindingObject, InputMask, OutputMask);
 					}
 					PRAGMA_DISABLE_DEPRECATION_WARNINGS
 					else
