ComputeFramework: Remove old deprecated functions from compute data providers.
#rnx

diff --git a/Engine/Plugins/Animation/DeformerGraph/Source/OptimusCore/Private/DataInterfaces/OptimusDataInterfaceAnimAttribute.cpp b/Engine/Plugins/Animation/DeformerGraph/Source/OptimusCore/Private/DataInterfaces/OptimusDataInterfaceAnimAttribute.cpp
--- a/Engine/Plugins/Animation/DeformerGraph/Source/OptimusCore/Private/DataInterfaces/OptimusDataInterfaceAnimAttribute.cpp
+++ b/Engine/Plugins/Animation/DeformerGraph/Source/OptimusCore/Private/DataInterfaces/OptimusDataInterfaceAnimAttribute.cpp
@@ -728,7 +728,7 @@ bool FOptimusAnimAttributeDataProviderProxy::IsValid(FValidationData const& InVa
 	return true;
 }
 
-void FOptimusAnimAttributeDataProviderProxy::AllocateResources(FRDGBuilder& GraphBuilder)
+void FOptimusAnimAttributeDataProviderProxy::AllocateResources(FRDGBuilder& GraphBuilder, FAllocationData const& InAllocationData)
 {
 	ArrayBuffers.Reset();
 	ArrayBufferSRVs.Reset();
diff --git a/Engine/Plugins/Animation/DeformerGraph/Source/OptimusCore/Private/DataInterfaces/OptimusDataInterfaceAnimAttribute.h b/Engine/Plugins/Animation/DeformerGraph/Source/OptimusCore/Private/DataInterfaces/OptimusDataInterfaceAnimAttribute.h
--- a/Engine/Plugins/Animation/DeformerGraph/Source/OptimusCore/Private/DataInterfaces/OptimusDataInterfaceAnimAttribute.h
+++ b/Engine/Plugins/Animation/DeformerGraph/Source/OptimusCore/Private/DataInterfaces/OptimusDataInterfaceAnimAttribute.h
@@ -233,7 +233,7 @@ public:
 
 	//~ Begin FComputeDataProviderRenderProxy Interface
 	bool IsValid(FValidationData const& InValidationData) const override;
-	void AllocateResources(FRDGBuilder& GraphBuilder) override;
+	void AllocateResources(FRDGBuilder& GraphBuilder, FAllocationData const& InAllocationData) override;
 	void GatherDispatchData(FDispatchData const& InDispatchData) override;
 	//~ End FComputeDataProviderRenderProxy Interface
 
diff --git a/Engine/Plugins/Animation/DeformerGraph/Source/OptimusCore/Private/DataInterfaces/OptimusDataInterfaceConnectivity.cpp b/Engine/Plugins/Animation/DeformerGraph/Source/OptimusCore/Private/DataInterfaces/OptimusDataInterfaceConnectivity.cpp
--- a/Engine/Plugins/Animation/DeformerGraph/Source/OptimusCore/Private/DataInterfaces/OptimusDataInterfaceConnectivity.cpp
+++ b/Engine/Plugins/Animation/DeformerGraph/Source/OptimusCore/Private/DataInterfaces/OptimusDataInterfaceConnectivity.cpp
@@ -253,7 +253,7 @@ bool FOptimusConnectivityDataProviderProxy::IsValid(FValidationData const& InVal
 	return true;
 }
 
-void FOptimusConnectivityDataProviderProxy::AllocateResources(FRDGBuilder& GraphBuilder)
+void FOptimusConnectivityDataProviderProxy::AllocateResources(FRDGBuilder& GraphBuilder, FAllocationData const& InAllocationData)
 {
 	const int32 LodIndex = SkeletalMeshObject->GetLOD();
 	TArray<uint32> const& ConnectivityData = AdjacencyBufferPerLod[LodIndex];
diff --git a/Engine/Plugins/Animation/DeformerGraph/Source/OptimusCore/Private/DataInterfaces/OptimusDataInterfaceConnectivity.h b/Engine/Plugins/Animation/DeformerGraph/Source/OptimusCore/Private/DataInterfaces/OptimusDataInterfaceConnectivity.h
--- a/Engine/Plugins/Animation/DeformerGraph/Source/OptimusCore/Private/DataInterfaces/OptimusDataInterfaceConnectivity.h
+++ b/Engine/Plugins/Animation/DeformerGraph/Source/OptimusCore/Private/DataInterfaces/OptimusDataInterfaceConnectivity.h
@@ -70,7 +70,7 @@ public:
 
 	//~ Begin FComputeDataProviderRenderProxy Interface
 	bool IsValid(FValidationData const& InValidationData) const override;
-	void AllocateResources(FRDGBuilder& GraphBuilder) override;
+	void AllocateResources(FRDGBuilder& GraphBuilder, FAllocationData const& InAllocationData) override;
 	void GatherDispatchData(FDispatchData const& InDispatchData);
 	//~ End FComputeDataProviderRenderProxy Interface
 
diff --git a/Engine/Plugins/Animation/DeformerGraph/Source/OptimusCore/Private/DataInterfaces/OptimusDataInterfaceDebugDraw.cpp b/Engine/Plugins/Animation/DeformerGraph/Source/OptimusCore/Private/DataInterfaces/OptimusDataInterfaceDebugDraw.cpp
--- a/Engine/Plugins/Animation/DeformerGraph/Source/OptimusCore/Private/DataInterfaces/OptimusDataInterfaceDebugDraw.cpp
+++ b/Engine/Plugins/Animation/DeformerGraph/Source/OptimusCore/Private/DataInterfaces/OptimusDataInterfaceDebugDraw.cpp
@@ -187,7 +187,7 @@ bool FOptimusDebugDrawDataProviderProxy::IsValid(FValidationData const& InValida
 	return true;
 }
 
-void FOptimusDebugDrawDataProviderProxy::AllocateResources(FRDGBuilder& GraphBuilder)
+void FOptimusDebugDrawDataProviderProxy::AllocateResources(FRDGBuilder& GraphBuilder, FAllocationData const& InAllocationData)
 {
 	// Allocate ShaderPrint output buffers.
 	FShaderPrintData ShaderPrintData = ShaderPrint::CreateShaderPrintData(GraphBuilder, Setup);
diff --git a/Engine/Plugins/Animation/DeformerGraph/Source/OptimusCore/Private/DataInterfaces/OptimusDataInterfaceDebugDraw.h b/Engine/Plugins/Animation/DeformerGraph/Source/OptimusCore/Private/DataInterfaces/OptimusDataInterfaceDebugDraw.h
--- a/Engine/Plugins/Animation/DeformerGraph/Source/OptimusCore/Private/DataInterfaces/OptimusDataInterfaceDebugDraw.h
+++ b/Engine/Plugins/Animation/DeformerGraph/Source/OptimusCore/Private/DataInterfaces/OptimusDataInterfaceDebugDraw.h
@@ -106,7 +106,7 @@ public:
 
 	//~ Begin FComputeDataProviderRenderProxy Interface
 	bool IsValid(FValidationData const& InValidationData) const override;
-	void AllocateResources(FRDGBuilder& GraphBuilder) override;
+	void AllocateResources(FRDGBuilder& GraphBuilder, FAllocationData const& InAllocationData) override;
 	void GatherDispatchData(FDispatchData const& InDispatchData) override;
 	//~ End FComputeDataProviderRenderProxy Interface
 
diff --git a/Engine/Plugins/Animation/DeformerGraph/Source/OptimusCore/Private/DataInterfaces/OptimusDataInterfaceGraph.cpp b/Engine/Plugins/Animation/DeformerGraph/Source/OptimusCore/Private/DataInterfaces/OptimusDataInterfaceGraph.cpp
--- a/Engine/Plugins/Animation/DeformerGraph/Source/OptimusCore/Private/DataInterfaces/OptimusDataInterfaceGraph.cpp
+++ b/Engine/Plugins/Animation/DeformerGraph/Source/OptimusCore/Private/DataInterfaces/OptimusDataInterfaceGraph.cpp
@@ -239,7 +239,7 @@ bool FOptimusGraphDataProviderProxy::IsValid(FValidationData const& InValidation
 	return true;
 }
 
-void FOptimusGraphDataProviderProxy::AllocateResources(FRDGBuilder& GraphBuilder)
+void FOptimusGraphDataProviderProxy::AllocateResources(FRDGBuilder& GraphBuilder, FAllocationData const& InAllocationData)
 {
 	ParameterArrayBuffers.Reset();
 	ParameterArrayBufferSRVs.Reset();
diff --git a/Engine/Plugins/Animation/DeformerGraph/Source/OptimusCore/Private/DataInterfaces/OptimusDataInterfaceGraph.h b/Engine/Plugins/Animation/DeformerGraph/Source/OptimusCore/Private/DataInterfaces/OptimusDataInterfaceGraph.h
--- a/Engine/Plugins/Animation/DeformerGraph/Source/OptimusCore/Private/DataInterfaces/OptimusDataInterfaceGraph.h
+++ b/Engine/Plugins/Animation/DeformerGraph/Source/OptimusCore/Private/DataInterfaces/OptimusDataInterfaceGraph.h
@@ -135,7 +135,7 @@ public:
 
 	//~ Begin FComputeDataProviderRenderProxy Interface
 	bool IsValid(FValidationData const& InValidationData) const override;
-	void AllocateResources(FRDGBuilder& GraphBuilder) override;
+	void AllocateResources(FRDGBuilder& GraphBuilder, FAllocationData const& InAllocationData) override;
 	void GatherDispatchData(FDispatchData const& InDispatchData) override;
 	//~ End FComputeDataProviderRenderProxy Interface
 
diff --git a/Engine/Plugins/Animation/DeformerGraph/Source/OptimusCore/Private/DataInterfaces/OptimusDataInterfaceHalfEdge.cpp b/Engine/Plugins/Animation/DeformerGraph/Source/OptimusCore/Private/DataInterfaces/OptimusDataInterfaceHalfEdge.cpp
--- a/Engine/Plugins/Animation/DeformerGraph/Source/OptimusCore/Private/DataInterfaces/OptimusDataInterfaceHalfEdge.cpp
+++ b/Engine/Plugins/Animation/DeformerGraph/Source/OptimusCore/Private/DataInterfaces/OptimusDataInterfaceHalfEdge.cpp
@@ -281,7 +281,7 @@ bool FOptimusHalfEdgeDataProviderProxy::IsValid(FValidationData const& InValidat
 	return true;
 }
 
-void FOptimusHalfEdgeDataProviderProxy::AllocateResources(FRDGBuilder& GraphBuilder)
+void FOptimusHalfEdgeDataProviderProxy::AllocateResources(FRDGBuilder& GraphBuilder, FAllocationData const& InAllocationData)
 {
 	FallbackSRV = GraphBuilder.CreateSRV(FRDGBufferSRVDesc(GSystemTextures.GetDefaultStructuredBuffer(GraphBuilder, sizeof(int32))));
 	
diff --git a/Engine/Plugins/Animation/DeformerGraph/Source/OptimusCore/Private/DataInterfaces/OptimusDataInterfaceHalfEdge.h b/Engine/Plugins/Animation/DeformerGraph/Source/OptimusCore/Private/DataInterfaces/OptimusDataInterfaceHalfEdge.h
--- a/Engine/Plugins/Animation/DeformerGraph/Source/OptimusCore/Private/DataInterfaces/OptimusDataInterfaceHalfEdge.h
+++ b/Engine/Plugins/Animation/DeformerGraph/Source/OptimusCore/Private/DataInterfaces/OptimusDataInterfaceHalfEdge.h
@@ -96,7 +96,7 @@ public:
 
 	//~ Begin FComputeDataProviderRenderProxy Interface
 	bool IsValid(FValidationData const& InValidationData) const override;
-	void AllocateResources(FRDGBuilder& GraphBuilder) override;
+	void AllocateResources(FRDGBuilder& GraphBuilder, FAllocationData const& InAllocationData) override;
 	void GatherDispatchData(FDispatchData const& InDispatchData);
 	//~ End FComputeDataProviderRenderProxy Interface
 
diff --git a/Engine/Plugins/Animation/DeformerGraph/Source/OptimusCore/Private/DataInterfaces/OptimusDataInterfaceRawBuffer.cpp b/Engine/Plugins/Animation/DeformerGraph/Source/OptimusCore/Private/DataInterfaces/OptimusDataInterfaceRawBuffer.cpp
--- a/Engine/Plugins/Animation/DeformerGraph/Source/OptimusCore/Private/DataInterfaces/OptimusDataInterfaceRawBuffer.cpp
+++ b/Engine/Plugins/Animation/DeformerGraph/Source/OptimusCore/Private/DataInterfaces/OptimusDataInterfaceRawBuffer.cpp
@@ -344,7 +344,7 @@ bool FOptimusTransientBufferDataProviderProxy::IsValid(FValidationData const& In
 	return true;
 }
 
-void FOptimusTransientBufferDataProviderProxy::AllocateResources(FRDGBuilder& GraphBuilder)
+void FOptimusTransientBufferDataProviderProxy::AllocateResources(FRDGBuilder& GraphBuilder, FAllocationData const& InAllocationData)
 {
 	// If we are using a raw type alias for the buffer then we need to adjust stride and count.
 	check(RawStride == 0 || ElementStride % RawStride == 0);
@@ -417,7 +417,7 @@ bool FOptimusImplicitPersistentBufferDataProviderProxy::IsValid(FValidationData
 	return true;
 }
 
-void FOptimusImplicitPersistentBufferDataProviderProxy::AllocateResources(FRDGBuilder& GraphBuilder)
+void FOptimusImplicitPersistentBufferDataProviderProxy::AllocateResources(FRDGBuilder& GraphBuilder, FAllocationData const& InAllocationData)
 {
 	TArray<int32> Count;
 	Count.Add(TotalElementCount);
diff --git a/Engine/Plugins/Animation/DeformerGraph/Source/OptimusCore/Private/DataInterfaces/OptimusDataInterfaceRawBuffer.h b/Engine/Plugins/Animation/DeformerGraph/Source/OptimusCore/Private/DataInterfaces/OptimusDataInterfaceRawBuffer.h
--- a/Engine/Plugins/Animation/DeformerGraph/Source/OptimusCore/Private/DataInterfaces/OptimusDataInterfaceRawBuffer.h
+++ b/Engine/Plugins/Animation/DeformerGraph/Source/OptimusCore/Private/DataInterfaces/OptimusDataInterfaceRawBuffer.h
@@ -295,7 +295,7 @@ public:
 
 	//~ Begin FComputeDataProviderRenderProxy Interface
 	bool IsValid(FValidationData const& InValidationData) const override;
-	void AllocateResources(FRDGBuilder& GraphBuilder) override;
+	void AllocateResources(FRDGBuilder& GraphBuilder, FAllocationData const& InAllocationData) override;
 	void GatherDispatchData(FDispatchData const& InDispatchData) override;
 	//~ End FComputeDataProviderRenderProxy Interface
 
@@ -329,7 +329,7 @@ public:
 
 	//~ Begin FComputeDataProviderRenderProxy Interface
 	bool IsValid(FValidationData const& InValidationData) const override;
-	void AllocateResources(FRDGBuilder& GraphBuilder) override;
+	void AllocateResources(FRDGBuilder& GraphBuilder, FAllocationData const& InAllocationData) override;
 	void GatherDispatchData(FDispatchData const& InDispatchData) override;
 	//~ End FComputeDataProviderRenderProxy Interface
 
diff --git a/Engine/Plugins/Animation/DeformerGraph/Source/OptimusCore/Private/DataInterfaces/OptimusDataInterfaceSkinnedMeshVertexAttribute.cpp b/Engine/Plugins/Animation/DeformerGraph/Source/OptimusCore/Private/DataInterfaces/OptimusDataInterfaceSkinnedMeshVertexAttribute.cpp
--- a/Engine/Plugins/Animation/DeformerGraph/Source/OptimusCore/Private/DataInterfaces/OptimusDataInterfaceSkinnedMeshVertexAttribute.cpp
+++ b/Engine/Plugins/Animation/DeformerGraph/Source/OptimusCore/Private/DataInterfaces/OptimusDataInterfaceSkinnedMeshVertexAttribute.cpp
@@ -119,21 +119,6 @@ UComputeDataProvider* UOptimusSkinnedMeshVertexAttributeDataInterface::CreateDat
 }
 
 
-bool UOptimusSkinnedMeshVertexAttributeDataProvider::IsValid() const
-{
-	if (SkinnedMeshComponent == nullptr || SkinnedMeshComponent->MeshObject == nullptr)
-	{
-		return false;
-	}
-	
-	const int32 LodIndex = SkinnedMeshComponent->MeshObject->GetLOD();
-	FSkeletalMeshRenderData const& SkeletalMeshRenderData = SkinnedMeshComponent->MeshObject->GetSkeletalMeshRenderData();
-	FSkeletalMeshLODRenderData const* LodRenderData = &SkeletalMeshRenderData.LODRenderData[LodIndex];
-
-	return LodRenderData->VertexAttributeBuffers.GetAttributeBuffer(AttributeName) != nullptr;
-}
-
-
 FComputeDataProviderRenderProxy* UOptimusSkinnedMeshVertexAttributeDataProvider::GetRenderProxy()
 {
 	if (const UOptimusSkinnedMeshVertexAttributeDataInterface* DataInterface = WeakDataInterface.Get())
diff --git a/Engine/Plugins/Animation/DeformerGraph/Source/OptimusCore/Private/DataInterfaces/OptimusDataInterfaceSkinnedMeshVertexAttribute.h b/Engine/Plugins/Animation/DeformerGraph/Source/OptimusCore/Private/DataInterfaces/OptimusDataInterfaceSkinnedMeshVertexAttribute.h
--- a/Engine/Plugins/Animation/DeformerGraph/Source/OptimusCore/Private/DataInterfaces/OptimusDataInterfaceSkinnedMeshVertexAttribute.h
+++ b/Engine/Plugins/Animation/DeformerGraph/Source/OptimusCore/Private/DataInterfaces/OptimusDataInterfaceSkinnedMeshVertexAttribute.h
@@ -70,7 +70,6 @@ public:
     TWeakObjectPtr<const UOptimusSkinnedMeshVertexAttributeDataInterface> WeakDataInterface;
 	
 	//~ Begin UComputeDataProvider Interface
-	bool IsValid() const override;
 	FComputeDataProviderRenderProxy* GetRenderProxy() override;
 	//~ End UComputeDataProvider Interface
 
diff --git a/Engine/Plugins/Animation/MLDeformer/MLDeformerFramework/Source/MLDeformerFramework/Private/DeformerGraphDataInterfaceSkeletonWithQuats.cpp b/Engine/Plugins/Animation/MLDeformer/MLDeformerFramework/Source/MLDeformerFramework/Private/DeformerGraphDataInterfaceSkeletonWithQuats.cpp
--- a/Engine/Plugins/Animation/MLDeformer/MLDeformerFramework/Source/MLDeformerFramework/Private/DeformerGraphDataInterfaceSkeletonWithQuats.cpp
+++ b/Engine/Plugins/Animation/MLDeformer/MLDeformerFramework/Source/MLDeformerFramework/Private/DeformerGraphDataInterfaceSkeletonWithQuats.cpp
@@ -217,7 +217,7 @@ namespace UE::MLDeformer
 		}
 	}
 
-	void FDeformerGraphSkeletonWithWeightedQuatsDataProviderProxy::AllocateResources(FRDGBuilder& GraphBuilder)
+	void FDeformerGraphSkeletonWithWeightedQuatsDataProviderProxy::AllocateResources(FRDGBuilder& GraphBuilder, FAllocationData const& InAllocationData)
 	{
 		TArray<FQuat4f> RefToLocalQuats;
 		CacheRefToLocalQuats(RefToLocalQuats);
diff --git a/Engine/Plugins/Animation/MLDeformer/MLDeformerFramework/Source/MLDeformerFramework/Private/DeformerGraphDataInterfaceSkeletonWithQuats.h b/Engine/Plugins/Animation/MLDeformer/MLDeformerFramework/Source/MLDeformerFramework/Private/DeformerGraphDataInterfaceSkeletonWithQuats.h
--- a/Engine/Plugins/Animation/MLDeformer/MLDeformerFramework/Source/MLDeformerFramework/Private/DeformerGraphDataInterfaceSkeletonWithQuats.h
+++ b/Engine/Plugins/Animation/MLDeformer/MLDeformerFramework/Source/MLDeformerFramework/Private/DeformerGraphDataInterfaceSkeletonWithQuats.h
@@ -70,7 +70,7 @@ namespace UE::MLDeformer
 		UE_API bool IsValid(FValidationData const& InValidationData) const override;
 		UE_API void GatherPermutations(FPermutationData& InOutPermutationData) const override;
 		UE_API void GatherDispatchData(FDispatchData const& InDispatchData) override;
-		UE_API void AllocateResources(FRDGBuilder& GraphBuilder) override;
+		UE_API void AllocateResources(FRDGBuilder& GraphBuilder, FAllocationData const& InAllocationData) override;
 		//~ End FComputeDataProviderRenderProxy Interface
 
 	private:
diff --git a/Engine/Plugins/Animation/MLDeformer/MLDeformerFramework/Source/MLDeformerFramework/Private/MLDeformerGraphDebugDataInterface.cpp b/Engine/Plugins/Animation/MLDeformer/MLDeformerFramework/Source/MLDeformerFramework/Private/MLDeformerGraphDebugDataInterface.cpp
--- a/Engine/Plugins/Animation/MLDeformer/MLDeformerFramework/Source/MLDeformerFramework/Private/MLDeformerGraphDebugDataInterface.cpp
+++ b/Engine/Plugins/Animation/MLDeformer/MLDeformerFramework/Source/MLDeformerFramework/Private/MLDeformerGraphDebugDataInterface.cpp
@@ -206,7 +206,7 @@ namespace UE::MLDeformer
 		return true;
 	}
 
-	void FMLDeformerGraphDebugDataProviderProxy::AllocateResources(FRDGBuilder& GraphBuilder)
+	void FMLDeformerGraphDebugDataProviderProxy::AllocateResources(FRDGBuilder& GraphBuilder, FAllocationData const& InAllocationData)
 	{
 		GroundTruthBuffer = GraphBuilder.CreateBuffer(FRDGBufferDesc::CreateStructuredDesc(sizeof(float), 3 * GroundTruthPositions.Num()), TEXT("MLDeformer.GroundTruthPositions"));
 		GroundTruthBufferSRV = GraphBuilder.CreateSRV(GroundTruthBuffer);
diff --git a/Engine/Plugins/Animation/MLDeformer/MLDeformerFramework/Source/MLDeformerFramework/Public/MLDeformerGraphDebugDataInterface.h b/Engine/Plugins/Animation/MLDeformer/MLDeformerFramework/Source/MLDeformerFramework/Public/MLDeformerGraphDebugDataInterface.h
--- a/Engine/Plugins/Animation/MLDeformer/MLDeformerFramework/Source/MLDeformerFramework/Public/MLDeformerGraphDebugDataInterface.h
+++ b/Engine/Plugins/Animation/MLDeformer/MLDeformerFramework/Source/MLDeformerFramework/Public/MLDeformerGraphDebugDataInterface.h
@@ -85,7 +85,7 @@ namespace UE::MLDeformer
 
 		// FComputeDataProviderRenderProxy overrides.
 		UE_API bool IsValid(FValidationData const& InValidationData) const override;
-		UE_API void AllocateResources(FRDGBuilder& GraphBuilder) override;
+		UE_API void AllocateResources(FRDGBuilder& GraphBuilder, FAllocationData const& InAllocationData) override;
 		UE_API void GatherDispatchData(FDispatchData const& InDispatchData) override;
 		// ~END FComputeDataProviderRenderProxy overrides.
 
diff --git a/Engine/Plugins/Animation/MLDeformer/NearestNeighborModel/Source/NearestNeighborModel/Private/OptimusDataInterfaceSkeletonWithQuats.cpp b/Engine/Plugins/Animation/MLDeformer/NearestNeighborModel/Source/NearestNeighborModel/Private/OptimusDataInterfaceSkeletonWithQuats.cpp
--- a/Engine/Plugins/Animation/MLDeformer/NearestNeighborModel/Source/NearestNeighborModel/Private/OptimusDataInterfaceSkeletonWithQuats.cpp
+++ b/Engine/Plugins/Animation/MLDeformer/NearestNeighborModel/Source/NearestNeighborModel/Private/OptimusDataInterfaceSkeletonWithQuats.cpp
@@ -221,7 +221,7 @@ void FOptimusSkeletonWithQuatsDataProviderProxy::GatherPermutations(FPermutation
 	}
 }
 
-void FOptimusSkeletonWithQuatsDataProviderProxy::AllocateResources(FRDGBuilder& GraphBuilder)
+void FOptimusSkeletonWithQuatsDataProviderProxy::AllocateResources(FRDGBuilder& GraphBuilder, FAllocationData const& InAllocationData)
 {
 	TArray<FQuat4f> RefToLocalQuats;
 	CacheRefToLocalQuats(RefToLocalQuats);
diff --git a/Engine/Plugins/Animation/MLDeformer/NearestNeighborModel/Source/NearestNeighborModel/Private/OptimusDataInterfaceSkeletonWithQuats.h b/Engine/Plugins/Animation/MLDeformer/NearestNeighborModel/Source/NearestNeighborModel/Private/OptimusDataInterfaceSkeletonWithQuats.h
--- a/Engine/Plugins/Animation/MLDeformer/NearestNeighborModel/Source/NearestNeighborModel/Private/OptimusDataInterfaceSkeletonWithQuats.h
+++ b/Engine/Plugins/Animation/MLDeformer/NearestNeighborModel/Source/NearestNeighborModel/Private/OptimusDataInterfaceSkeletonWithQuats.h
@@ -66,7 +66,7 @@ public:
 	bool IsValid(FValidationData const& InValidationData) const override;
 	void GatherPermutations(FPermutationData& InOutPermutationData) const override;
 	void GatherDispatchData(FDispatchData const& InDispatchData) override;
-	void AllocateResources(FRDGBuilder& GraphBuilder) override;
+	void AllocateResources(FRDGBuilder& GraphBuilder, FAllocationData const& InAllocationData) override;
 	//~ End FComputeDataProviderRenderProxy Interface
 
 private:
diff --git a/Engine/Plugins/Animation/MLDeformer/VertexDeltaModel/Source/VertexDeltaModel/Private/VertexDeltaGraphDataInterface.cpp b/Engine/Plugins/Animation/MLDeformer/VertexDeltaModel/Source/VertexDeltaModel/Private/VertexDeltaGraphDataInterface.cpp
--- a/Engine/Plugins/Animation/MLDeformer/VertexDeltaModel/Source/VertexDeltaModel/Private/VertexDeltaGraphDataInterface.cpp
+++ b/Engine/Plugins/Animation/MLDeformer/VertexDeltaModel/Source/VertexDeltaModel/Private/VertexDeltaGraphDataInterface.cpp
@@ -142,7 +142,7 @@ namespace UE::VertexDeltaModel
 	}
 	
 
-	void FVertexDeltaGraphDataProviderProxy::AllocateResources(FRDGBuilder& GraphBuilder)
+	void FVertexDeltaGraphDataProviderProxy::AllocateResources(FRDGBuilder& GraphBuilder, FAllocationData const& InAllocationData)
 	{
 		if (VertexDeltaModelInstance && VertexDeltaModelInstance->IsValidForDataProvider())
 		{
diff --git a/Engine/Plugins/Animation/MLDeformer/VertexDeltaModel/Source/VertexDeltaModel/Private/VertexDeltaGraphDebugDataInterface.cpp b/Engine/Plugins/Animation/MLDeformer/VertexDeltaModel/Source/VertexDeltaModel/Private/VertexDeltaGraphDebugDataInterface.cpp
--- a/Engine/Plugins/Animation/MLDeformer/VertexDeltaModel/Source/VertexDeltaModel/Private/VertexDeltaGraphDebugDataInterface.cpp
+++ b/Engine/Plugins/Animation/MLDeformer/VertexDeltaModel/Source/VertexDeltaModel/Private/VertexDeltaGraphDebugDataInterface.cpp
@@ -198,7 +198,7 @@ namespace UE::VertexDeltaModel
 		return true;
 	}
 
-	void FVertexDeltaGraphDebugDataProviderProxy::AllocateResources(FRDGBuilder& GraphBuilder)
+	void FVertexDeltaGraphDebugDataProviderProxy::AllocateResources(FRDGBuilder& GraphBuilder, FAllocationData const& InAllocationData)
 	{
 		GroundTruthBuffer = GraphBuilder.CreateBuffer(FRDGBufferDesc::CreateStructuredDesc(sizeof(float), 3 * GroundTruthPositions.Num()), TEXT("MLDeformer.GroundTruthPositions"));
 		GroundTruthBufferSRV = GraphBuilder.CreateSRV(GroundTruthBuffer);
diff --git a/Engine/Plugins/Animation/MLDeformer/VertexDeltaModel/Source/VertexDeltaModel/Public/VertexDeltaGraphDataInterface.h b/Engine/Plugins/Animation/MLDeformer/VertexDeltaModel/Source/VertexDeltaModel/Public/VertexDeltaGraphDataInterface.h
--- a/Engine/Plugins/Animation/MLDeformer/VertexDeltaModel/Source/VertexDeltaModel/Public/VertexDeltaGraphDataInterface.h
+++ b/Engine/Plugins/Animation/MLDeformer/VertexDeltaModel/Source/VertexDeltaModel/Public/VertexDeltaGraphDataInterface.h
@@ -80,7 +80,7 @@ namespace UE::VertexDeltaModel
 		// FComputeDataProviderRenderProxy overrides.
 		bool IsValid(FValidationData const& InValidationData) const override;
 		void GatherDispatchData(FDispatchData const& InDispatchData) override;
-		void AllocateResources(FRDGBuilder& GraphBuilder) override;
+		void AllocateResources(FRDGBuilder& GraphBuilder, FAllocationData const& InAllocationData) override;
 		// ~END FComputeDataProviderRenderProxy overrides.
 
 	protected:
diff --git a/Engine/Plugins/Animation/MLDeformer/VertexDeltaModel/Source/VertexDeltaModel/Public/VertexDeltaGraphDebugDataInterface.h b/Engine/Plugins/Animation/MLDeformer/VertexDeltaModel/Source/VertexDeltaModel/Public/VertexDeltaGraphDebugDataInterface.h
--- a/Engine/Plugins/Animation/MLDeformer/VertexDeltaModel/Source/VertexDeltaModel/Public/VertexDeltaGraphDebugDataInterface.h
+++ b/Engine/Plugins/Animation/MLDeformer/VertexDeltaModel/Source/VertexDeltaModel/Public/VertexDeltaGraphDebugDataInterface.h
@@ -84,7 +84,7 @@ namespace UE::VertexDeltaModel
 
 		// FComputeDataProviderRenderProxy overrides.
 		bool IsValid(FValidationData const& InValidationData) const override;
-		void AllocateResources(FRDGBuilder& GraphBuilder) override;
+		void AllocateResources(FRDGBuilder& GraphBuilder, FAllocationData const& InAllocationData) override;
 		void GatherDispatchData(FDispatchData const& InDispatchData) override;
 		// ~END FComputeDataProviderRenderProxy overrides.
 
diff --git a/Engine/Plugins/Experimental/ChaosFlesh/Source/ChaosFleshEngine/Private/DataInterfaces/DIFleshDeformer.cpp b/Engine/Plugins/Experimental/ChaosFlesh/Source/ChaosFleshEngine/Private/DataInterfaces/DIFleshDeformer.cpp
--- a/Engine/Plugins/Experimental/ChaosFlesh/Source/ChaosFleshEngine/Private/DataInterfaces/DIFleshDeformer.cpp
+++ b/Engine/Plugins/Experimental/ChaosFlesh/Source/ChaosFleshEngine/Private/DataInterfaces/DIFleshDeformer.cpp
@@ -250,12 +250,6 @@ UDIFleshDeformer::CreateDataProvider(TObjectPtr<UObject> InBinding, uint64 InInp
 // DataProvider
 //
 
-bool
-UDIFleshDeformerDataProvider::IsValid() const
-{
-	return (SkinnedMesh && FleshMesh) || (StaticMesh && FleshMesh);
-}
-
 FComputeDataProviderRenderProxy*
 UDIFleshDeformerDataProvider::GetRenderProxy()
 {
@@ -347,7 +341,7 @@ FDIFleshDeformerProviderProxy::GatherPermutations(FPermutationData& InOutPermuta
 }
 
 void
-FDIFleshDeformerProviderProxy::AllocateResources(FRDGBuilder& GraphBuilder)
+FDIFleshDeformerProviderProxy::AllocateResources(FRDGBuilder& GraphBuilder, FAllocationData const& InAllocationData)
 {}
 
 void
diff --git a/Engine/Plugins/Experimental/ChaosFlesh/Source/ChaosFleshEngine/Public/DataInterfaces/DIFleshDeformer.h b/Engine/Plugins/Experimental/ChaosFlesh/Source/ChaosFleshEngine/Public/DataInterfaces/DIFleshDeformer.h
--- a/Engine/Plugins/Experimental/ChaosFlesh/Source/ChaosFleshEngine/Public/DataInterfaces/DIFleshDeformer.h
+++ b/Engine/Plugins/Experimental/ChaosFlesh/Source/ChaosFleshEngine/Public/DataInterfaces/DIFleshDeformer.h
@@ -76,7 +76,6 @@ public:
 	FFleshDeformerParameters FleshDeformerParameters;
 
 	//~ Begin UComputeDataProvider Interface
-	bool IsValid() const override;
 	FComputeDataProviderRenderProxy* GetRenderProxy() override;
 	//~ End UComputeDataProvider Interface
 };
@@ -89,7 +88,7 @@ public:
 	//~ Begin FComputeDataProviderRenderProxy Interface
 	bool IsValid(FValidationData const& InValidationData) const override;
 	void GatherPermutations(FPermutationData& InOutPermutationData) const override;
-	void AllocateResources(FRDGBuilder& GraphBuilder) override;
+	void AllocateResources(FRDGBuilder& GraphBuilder, FAllocationData const& InAllocationData) override;
 	void GatherDispatchData(FDispatchData const& InDispatchData);
 	//~ End FComputeDataProviderRenderProxy Interface
 
diff --git a/Engine/Plugins/Runtime/ComputeFramework/Source/ComputeFramework/Public/ComputeFramework/ComputeDataProvider.h b/Engine/Plugins/Runtime/ComputeFramework/Source/ComputeFramework/Public/ComputeFramework/ComputeDataProvider.h
--- a/Engine/Plugins/Runtime/ComputeFramework/Source/ComputeFramework/Public/ComputeFramework/ComputeDataProvider.h
+++ b/Engine/Plugins/Runtime/ComputeFramework/Source/ComputeFramework/Public/ComputeFramework/ComputeDataProvider.h
@@ -27,10 +27,6 @@ public:
 	/** Set up the data provider from the given data interface and binding. */
 	virtual void Initialize(const UComputeDataInterface* InDataInterface, UObject* InBinding, uint64 InInputMask, uint64 InOutputMask) {}
 
-	/** Return false if the provider has not been fully initialized. */
-	UE_DEPRECATED(5.2, "Implement any validation in FComputeDataProviderRenderProxy::IsValid().")
-	virtual bool IsValid() const { return true; }
-	
 	/**
 	 * Get an associated render thread proxy object.
 	 * Currently these are created and destroyed per frame by the FComputeGraphInstance.
@@ -91,17 +87,8 @@ public:
 		FRDGExternalAccessQueue& ExternalAccessQueue;
 	};
 
-	// todo: Fix up derived classes to implement new AllocateResources().
-	UE_DEPRECATED(5.3, "Convert to using the new AllocateResources() that takes FAllocationData.")
-	virtual void AllocateResources(FRDGBuilder& GraphBuilder) {}
-
 	/* Called once before any calls to GatherDispatchData() to allow any RDG resource allocation. */
-	virtual void AllocateResources(FRDGBuilder& GraphBuilder, FAllocationData const& InAllocationData)
-	{
-		PRAGMA_DISABLE_DEPRECATION_WARNINGS
-		AllocateResources(GraphBuilder);
-		PRAGMA_ENABLE_DEPRECATION_WARNINGS
-	}
+	virtual void AllocateResources(FRDGBuilder& GraphBuilder, FAllocationData const& InAllocationData) {}
 
 	/** Setup needed to gather dispatch data. */
 	struct FDispatchData
@@ -140,6 +127,9 @@ public:
 	/** Called immediately after every kernel dispatch that outputs to this data. */
 	virtual void PostSubmit(FRDGBuilder& InGraphBuilder) const {}
 
+	UE_DEPRECATED(5.6, "Use PostSubmit instead.")
+	virtual void PostGraphDispatch(FRDGBuilder& GraphBuilder) const {}
+
 protected:
 	/** Helper for making an FStridedView over the FDispatchData. */
 	template <typename ElementType>
@@ -147,36 +137,4 @@ protected:
 	{
 		return TStridedView<ElementType>(InDispatchData.ParameterBufferStride, (ElementType*)(InDispatchData.ParameterBuffer + InDispatchData.ParameterBufferOffset), InDispatchData.NumInvocations);
 	}
-
-protected:
-	//~ Start deprecation of old GatherDispatchData() interface.
-	struct FDispatchSetup
-	{
-		UE_DEPRECATED(5.2, "Convert to using the new GatherDispatchData() that takes FDispatchData.")
-		int32 NumInvocations;
-		UE_DEPRECATED(5.2, "Convert to using the new GatherDispatchData() that takes FDispatchData.")
-		int32 ParameterBufferOffset;
-		UE_DEPRECATED(5.2, "Convert to using the new GatherDispatchData() that takes FDispatchData.")
-		int32 ParameterBufferStride;
-		UE_DEPRECATED(5.2, "Convert to using the new GatherDispatchData() that takes FDispatchData.")
-		int32 ParameterStructSizeForValidation;
-		UE_DEPRECATED(5.2, "Convert to using the new GatherDispatchData() that takes FDispatchData.")
-		FComputeKernelPermutationVector const& PermutationVector;
-	};
-
-	struct FCollectedDispatchData
-	{
-		UE_DEPRECATED(5.2, "Convert to using the new GatherDispatchData() that takes FDispatchData.")
-		uint8* ParameterBuffer;
-		UE_DEPRECATED(5.2, "Convert to using the new GatherDispatchData() that takes FDispatchData.")
-		TArray<int32> PermutationId;
-	};
-
-	UE_DEPRECATED(5.2, "Convert to using the new GatherDispatchData() that takes FDispatchData.")
-	virtual void GatherDispatchData(FDispatchSetup const& InDispatchSetup, FCollectedDispatchData& InOutDispatchData) {}
-	//~ End deprecation of old GatherDispatchData() interface.
-
-public:
-	UE_DEPRECATED(5.6, "Use PostSubmit instead.")
-	virtual void PostGraphDispatch(FRDGBuilder& GraphBuilder) const {}
 };
