[ComputeFramework] Add execution group for view-dependent work.

Pass view information into compute graph worker.

#rnx

diff --git a/Engine/Plugins/PCG/Source/PCG/Private/Compute/DataInterfaces/PCGInstanceDataInterface.cpp b/Engine/Plugins/PCG/Source/PCG/Private/Compute/DataInterfaces/PCGInstanceDataInterface.cpp
--- a/Engine/Plugins/PCG/Source/PCG/Private/Compute/DataInterfaces/PCGInstanceDataInterface.cpp
+++ b/Engine/Plugins/PCG/Source/PCG/Private/Compute/DataInterfaces/PCGInstanceDataInterface.cpp
@@ -9,6 +9,7 @@
 #include "Compute/PCGComputeKernel.h"
 #include "Compute/PCGDataBinding.h"
 
+#include "ComputeWorkerInterface.h"
 #include "GlobalRenderResources.h"
 #include "GlobalShader.h"
 #include "GPUSceneWriter.h"
@@ -289,7 +290,7 @@ void FPCGInstanceDataProviderProxy::GatherDispatchData(FDispatchData const& InDi
 	}
 }
 
-void FPCGInstanceDataProviderProxy::PostSubmit(FRDGBuilder& InGraphBuilder) const
+void FPCGInstanceDataProviderProxy::PostSubmit(FComputeContext& Context) const
 {
 	TRACE_CPUPROFILER_EVENT_SCOPE(FPCGInstanceDataProviderProxy::PostSubmit);
 
@@ -298,9 +299,9 @@ void FPCGInstanceDataProviderProxy::PostSubmit(FRDGBuilder& InGraphBuilder) cons
 		return;
 	}
 
-	const TRefCountPtr<FRDGPooledBuffer> InstanceDataExported = InGraphBuilder.ConvertToExternalBuffer(InstanceData);
-	const TRefCountPtr<FRDGPooledBuffer> InstanceCustomFloatDataExported = InGraphBuilder.ConvertToExternalBuffer(InstanceCustomFloatData);
-	const TRefCountPtr<FRDGPooledBuffer> WriteCountersExported = InGraphBuilder.ConvertToExternalBuffer(WriteCounters);
+	const TRefCountPtr<FRDGPooledBuffer> InstanceDataExported = Context.GraphBuilder.ConvertToExternalBuffer(InstanceData);
+	const TRefCountPtr<FRDGPooledBuffer> InstanceCustomFloatDataExported = Context.GraphBuilder.ConvertToExternalBuffer(InstanceCustomFloatData);
+	const TRefCountPtr<FRDGPooledBuffer> WriteCountersExported = Context.GraphBuilder.ConvertToExternalBuffer(WriteCounters);
 	
 	FTSTicker::GetCoreTicker().AddTicker(TEXT("ApplyPrimitiveSceneUpdates"), 0.0f, [DataProviderWeak=DataProvider, DataProviderGeneration=DataProviderGeneration, InstanceDataExported, InstanceCustomFloatDataExported, WriteCountersExported](float)
 	{
diff --git a/Engine/Plugins/PCG/Source/PCG/Private/Compute/DataInterfaces/PCGInstanceDataInterface.h b/Engine/Plugins/PCG/Source/PCG/Private/Compute/DataInterfaces/PCGInstanceDataInterface.h
--- a/Engine/Plugins/PCG/Source/PCG/Private/Compute/DataInterfaces/PCGInstanceDataInterface.h
+++ b/Engine/Plugins/PCG/Source/PCG/Private/Compute/DataInterfaces/PCGInstanceDataInterface.h
@@ -78,7 +78,7 @@ public:
 	bool IsValid(FValidationData const& InValidationData) const override;
 	void AllocateResources(FRDGBuilder& GraphBuilder, FAllocationData const& InAllocationData);
 	void GatherDispatchData(FDispatchData const& InDispatchData) override;
-	void PostSubmit(FRDGBuilder& InGraphBuilder) const override;
+	void PostSubmit(FComputeContext& Context) const override;
 	//~ End FComputeDataProviderRenderProxy Interface
 
 protected:
diff --git a/Engine/Plugins/Runtime/ComputeFramework/Source/ComputeFramework/Private/ComputeFramework.cpp b/Engine/Plugins/Runtime/ComputeFramework/Source/ComputeFramework/Private/ComputeFramework.cpp
--- a/Engine/Plugins/Runtime/ComputeFramework/Source/ComputeFramework/Private/ComputeFramework.cpp
+++ b/Engine/Plugins/Runtime/ComputeFramework/Source/ComputeFramework/Private/ComputeFramework.cpp
@@ -89,10 +89,20 @@ namespace ComputeFramework
 			UE::RenderCommandPipe::FSyncScope SyncScope({ &UE::RenderCommandPipe::SkeletalMesh });
 
 			ENQUEUE_RENDER_COMMAND(ComputeFrameworkFlushCommand)(
-				[ComputeGraphWorker, InExecutionGroupName](FRHICommandListImmediate& RHICmdList)
+				[ComputeGraphWorker, InExecutionGroupName, InScene](FRHICommandListImmediate& RHICmdList)
 				{
 					FRDGBuilder GraphBuilder(RHICmdList);
-					ComputeGraphWorker->SubmitWork(GraphBuilder, InExecutionGroupName, GMaxRHIFeatureLevel);
+
+					FComputeContext Context
+					{
+						.GraphBuilder = GraphBuilder,
+						.ExecutionGroupName = InExecutionGroupName,
+						.FeatureLevel = GMaxRHIFeatureLevel,
+						.Scene = InScene,
+					};
+
+					ComputeGraphWorker->SubmitWork(Context);
+
 					GraphBuilder.Execute();
 				});
 		}
diff --git a/Engine/Plugins/Runtime/ComputeFramework/Source/ComputeFramework/Private/ComputeFramework/ComputeDataProvider.cpp b/Engine/Plugins/Runtime/ComputeFramework/Source/ComputeFramework/Private/ComputeFramework/ComputeDataProvider.cpp
new file mode 100644
--- /dev/null
+++ b/Engine/Plugins/Runtime/ComputeFramework/Source/ComputeFramework/Private/ComputeFramework/ComputeDataProvider.cpp
@@ -0,0 +1,19 @@
+// Copyright Epic Games, Inc. All Rights Reserved.
+
+#include "ComputeFramework/ComputeDataProvider.h"
+
+#include "ComputeWorkerInterface.h"
+
+#include UE_INLINE_GENERATED_CPP_BY_NAME(ComputeDataProvider)
+
+void FComputeDataProviderRenderProxy::PreSubmit(FRDGBuilder& InGraphBuilder) const
+{
+	FComputeContext Context = { .GraphBuilder = InGraphBuilder };
+	PreSubmit(Context);
+}
+
+void FComputeDataProviderRenderProxy::PostSubmit(FRDGBuilder& InGraphBuilder) const
+{
+	FComputeContext Context = { .GraphBuilder = InGraphBuilder };
+	PostSubmit(Context);
+}
diff --git a/Engine/Plugins/Runtime/ComputeFramework/Source/ComputeFramework/Private/ComputeFramework/ComputeGraphWorker.h b/Engine/Plugins/Runtime/ComputeFramework/Source/ComputeFramework/Private/ComputeFramework/ComputeGraphWorker.h
--- a/Engine/Plugins/Runtime/ComputeFramework/Source/ComputeFramework/Private/ComputeFramework/ComputeGraphWorker.h
+++ b/Engine/Plugins/Runtime/ComputeFramework/Source/ComputeFramework/Private/ComputeFramework/ComputeGraphWorker.h
@@ -37,10 +37,7 @@ public:
 	bool HasWork(FName InExecutionGroupName) const override;
 
 	/** Submit enqueued compute graph work. */
-	void SubmitWork(
-		FRDGBuilder& GraphBuilder,
-		FName InExecutionGroupName,
-		ERHIFeatureLevel::Type InFeatureLevel ) override;
+	void SubmitWork(FComputeContext& Context) override;
 
 private:
 	/** Poll active readbacks and execute callbacks for ready ones. */
diff --git a/Engine/Plugins/Runtime/ComputeFramework/Source/ComputeFramework/Private/ComputeFramework/ComputeSystem.h b/Engine/Plugins/Runtime/ComputeFramework/Source/ComputeFramework/Private/ComputeFramework/ComputeSystem.h
--- a/Engine/Plugins/Runtime/ComputeFramework/Source/ComputeFramework/Private/ComputeFramework/ComputeSystem.h
+++ b/Engine/Plugins/Runtime/ComputeFramework/Source/ComputeFramework/Private/ComputeFramework/ComputeSystem.h
@@ -6,6 +6,7 @@
 #include "Containers/Map.h"
 
 class FComputeGraphTaskWorker;
+class FComputeViewExtension;
 
 /** Implementation of IComputeSystem which manages FComputeGraphTaskWorker objects. */
 class FComputeFrameworkSystem : public IComputeSystem
@@ -20,4 +21,6 @@ public:
 
 private:
 	TMap<FSceneInterface const*, FComputeGraphTaskWorker*> ComputeWorkers;
+
+	TSharedPtr<FComputeViewExtension, ESPMode::ThreadSafe> ViewExtension;
 };
diff --git a/Engine/Plugins/Runtime/ComputeFramework/Source/ComputeFramework/Private/ComputeFramework/ComputeViewExtension.cpp b/Engine/Plugins/Runtime/ComputeFramework/Source/ComputeFramework/Private/ComputeFramework/ComputeViewExtension.cpp
new file mode 100644
--- /dev/null
+++ b/Engine/Plugins/Runtime/ComputeFramework/Source/ComputeFramework/Private/ComputeFramework/ComputeViewExtension.cpp
@@ -0,0 +1,69 @@
+// Copyright Epic Games, Inc. All Rights Reserved.
+
+#include "ComputeViewExtension.h"
+
+#include "ComputeWorkerInterface.h"
+
+FComputeViewExtension::FComputeViewExtension(const FAutoRegister& AutoReg, UWorld* InWorld)
+	: FWorldSceneViewExtension(AutoReg, InWorld)
+{
+}
+
+void FComputeViewExtension::PostTLASBuild_RenderThread(FRDGBuilder& GraphBuilder, FSceneView& InView)
+{
+	// Verified this gets called, if a graph that requires it sets it up. If we're happy with these changes, I can add a new PostTLASBuild function to the view extension interface?
+	const TWeakObjectPtr<UWorld> WorldPtr = GetWorld();
+	if (!WorldPtr.IsValid())
+	{
+		return;
+	}
+
+	const FSceneInterface* Scene = WorldPtr->Scene;
+	if (!Scene)
+	{
+		return;
+	}
+
+	const ERHIFeatureLevel::Type FeatureLevel = Scene->GetFeatureLevel();
+
+	ComputeWorkers.Empty(ComputeWorkers.Num());
+	Scene->GetComputeTaskWorkers(ComputeWorkers);
+
+	for (IComputeTaskWorker* ComputeTaskWorker : ComputeWorkers)
+	{
+		if (ComputeTaskWorker->HasWork(ComputeTaskExecutionGroup::PostTLASBuild))
+		{
+			FComputeContext Context
+			{
+				.GraphBuilder = GraphBuilder,
+				.ExecutionGroupName = ComputeTaskExecutionGroup::PostTLASBuild,
+				.FeatureLevel = FeatureLevel,
+				.Scene = Scene,
+				.View = &InView,
+			};
+
+			ComputeTaskWorker->SubmitWork(Context);
+		}
+	}
+}
+
+ESceneViewExtensionFlags FComputeViewExtension::GetFlags() const
+{
+	const TWeakObjectPtr<UWorld> WorldPtr = GetWorld();
+
+	if (const FSceneInterface* Scene = WorldPtr.IsValid() ? WorldPtr.Get()->Scene : nullptr)
+	{
+		ComputeWorkers.Empty(ComputeWorkers.Num());
+		Scene->GetComputeTaskWorkers(ComputeWorkers);
+
+		for (IComputeTaskWorker* ComputeTaskWorker : ComputeWorkers)
+		{
+			if (ComputeTaskWorker->HasWork(ComputeTaskExecutionGroup::PostTLASBuild))
+			{
+				return ESceneViewExtensionFlags::SubscribesToPostTLASBuild;
+			}
+		}
+	}
+
+	return ESceneViewExtensionFlags::None;
+}
diff --git a/Engine/Plugins/Runtime/ComputeFramework/Source/ComputeFramework/Private/ComputeFramework/ComputeViewExtension.h b/Engine/Plugins/Runtime/ComputeFramework/Source/ComputeFramework/Private/ComputeFramework/ComputeViewExtension.h
new file mode 100644
--- /dev/null
+++ b/Engine/Plugins/Runtime/ComputeFramework/Source/ComputeFramework/Private/ComputeFramework/ComputeViewExtension.h
@@ -0,0 +1,23 @@
+// Copyright Epic Games, Inc. All Rights Reserved.
+
+#pragma once
+
+#include "SceneViewExtension.h"
+
+/*
+ * Compute view extension, useful for receiving callbacks from renderer.
+ */
+class FComputeViewExtension : public FWorldSceneViewExtension
+{
+public:
+	FComputeViewExtension(const FAutoRegister& AutoReg, UWorld* InWorld);
+
+	//~ Begin FSceneViewExtensionBase implementation
+	virtual void PostTLASBuild_RenderThread(FRDGBuilder& GraphBuilder, FSceneView& InView) override;
+	virtual ESceneViewExtensionFlags GetFlags() const override;
+	//~ End FSceneViewExtensionBase implementation
+
+private:
+	// Local to methods but stored here and emptied at each usage to avoid per frame array allocations.
+	mutable TArray<IComputeTaskWorker*> ComputeWorkers;
+};
diff --git a/Engine/Plugins/Runtime/ComputeFramework/Source/ComputeFramework/Private/ComputeGraphWorker.cpp b/Engine/Plugins/Runtime/ComputeFramework/Source/ComputeFramework/Private/ComputeGraphWorker.cpp
--- a/Engine/Plugins/Runtime/ComputeFramework/Source/ComputeFramework/Private/ComputeGraphWorker.cpp
+++ b/Engine/Plugins/Runtime/ComputeFramework/Source/ComputeFramework/Private/ComputeGraphWorker.cpp
@@ -98,23 +98,25 @@ bool FComputeGraphTaskWorker::HasWork(FName InExecutionGroupName) const
 	return GraphInvocations != nullptr && GraphInvocations->Num();
 }
 
-void FComputeGraphTaskWorker::SubmitWork(FRDGBuilder& GraphBuilder, FName InExecutionGroupName, ERHIFeatureLevel::Type FeatureLevel)
+void FComputeGraphTaskWorker::SubmitWork(FComputeContext& Context)
 {
 	TRACE_CPUPROFILER_EVENT_SCOPE(ComputeFramework::ExecuteBatches);
-	RDG_EVENT_SCOPE_STAT(GraphBuilder, ComputeFramework_ExecuteBatches, "ComputeFramework::ExecuteBatches");
-	RDG_GPU_STAT_SCOPE(GraphBuilder, ComputeFramework_ExecuteBatches);
+	RDG_EVENT_SCOPE_STAT(Context.GraphBuilder, ComputeFramework_ExecuteBatches, "ComputeFramework::ExecuteBatches");
+	RDG_GPU_STAT_SCOPE(Context.GraphBuilder, ComputeFramework_ExecuteBatches);
+
+	FRDGBuilder& GraphBuilder = Context.GraphBuilder;
 
 	const bool bCaptureAllWork = GTriggerGPUCaptureDispatches > 0;
 	RenderCaptureInterface::FScopedCapture RenderCapture(bCaptureAllWork, GraphBuilder, TEXT("FComputeGraphTaskWorker::SubmitWork"));
 	GTriggerGPUCaptureDispatches = FMath::Max(GTriggerGPUCaptureDispatches - 1, 0);
 
 	// Currently poll readbacks once at end of frame.
-	if (InExecutionGroupName == ComputeTaskExecutionGroup::EndOfFrameUpdate)
+	if (Context.ExecutionGroupName == ComputeTaskExecutionGroup::EndOfFrameUpdate)
 	{
 		UpdateReadbacks();
 	}
 
-	TArray<FGraphInvocation> const* FoundGraphInvocations = GraphInvocationsPerGroup.Find(InExecutionGroupName);
+	TArray<FGraphInvocation> const* FoundGraphInvocations = GraphInvocationsPerGroup.Find(Context.ExecutionGroupName);
 	if (!FoundGraphInvocations || FoundGraphInvocations->Num() == 0)
 	{
 		return;
@@ -261,7 +263,7 @@ void FComputeGraphTaskWorker::SubmitWork(FRDGBuilder& GraphBuilder, FName InExec
 		{
 			if (FComputeDataProviderRenderProxy* DataProviderProxy = GraphInvocation.DataProviderRenderProxies[DataProviderIndex])
 			{
-				DataProviderProxy->PreSubmit(GraphBuilder);
+				DataProviderProxy->PreSubmit(Context);
 			}
 		}
 
@@ -337,7 +339,7 @@ void FComputeGraphTaskWorker::SubmitWork(FRDGBuilder& GraphBuilder, FName InExec
 		{
 			if (FComputeDataProviderRenderProxy* DataProviderProxy = GraphInvocation.DataProviderRenderProxies[DataProviderIndex])
 			{
-				DataProviderProxy->PostSubmit(GraphBuilder);
+				DataProviderProxy->PostSubmit(Context);
 			}
 		}
 
@@ -375,9 +377,9 @@ void FComputeGraphTaskWorker::SubmitWork(FRDGBuilder& GraphBuilder, FName InExec
 
 	// Release any graph resources at the end of graph execution.
 	GraphBuilder.AddPostExecuteCallback(
-		[this, InExecutionGroupName]
+		[this, ExecutionGroupName=Context.ExecutionGroupName]
 		{
-			GraphInvocationsPerGroup.FindChecked(InExecutionGroupName).Reset();
+			GraphInvocationsPerGroup.FindChecked(ExecutionGroupName).Reset();
 		});
 }
 
diff --git a/Engine/Plugins/Runtime/ComputeFramework/Source/ComputeFramework/Private/ComputeSystem.cpp b/Engine/Plugins/Runtime/ComputeFramework/Source/ComputeFramework/Private/ComputeSystem.cpp
--- a/Engine/Plugins/Runtime/ComputeFramework/Source/ComputeFramework/Private/ComputeSystem.cpp
+++ b/Engine/Plugins/Runtime/ComputeFramework/Source/ComputeFramework/Private/ComputeSystem.cpp
@@ -3,12 +3,18 @@
 #include "ComputeFramework/ComputeSystem.h"
 
 #include "ComputeFramework/ComputeGraphWorker.h"
+#include "ComputeFramework/ComputeViewExtension.h"
 
 void FComputeFrameworkSystem::CreateWorkers(FSceneInterface const* InScene, TArray<IComputeTaskWorker*>& OutWorkers)
 {
 	FComputeGraphTaskWorker* ComputeWorker = new FComputeGraphTaskWorker();
 	ComputeWorkers.Add(InScene, ComputeWorker);
 	OutWorkers.Add(ComputeWorker);
+
+	if (UWorld* World = InScene->GetWorld())
+	{
+		ViewExtension = FSceneViewExtensions::NewExtension<FComputeViewExtension>(World);
+	}
 }
 
 void FComputeFrameworkSystem::DestroyWorkers(FSceneInterface const* InScene, TArray<IComputeTaskWorker*>& InOutWorkers)
@@ -24,6 +30,8 @@ void FComputeFrameworkSystem::DestroyWorkers(FSceneInterface const* InScene, TAr
 
 		ComputeWorkers.Remove(InScene);
 	}
+
+	ViewExtension.Reset();
 }
 
 FComputeGraphTaskWorker* FComputeFrameworkSystem::GetComputeWorker(FSceneInterface const* InScene) const
diff --git a/Engine/Plugins/Runtime/ComputeFramework/Source/ComputeFramework/Public/ComputeFramework/ComputeDataProvider.h b/Engine/Plugins/Runtime/ComputeFramework/Source/ComputeFramework/Public/ComputeFramework/ComputeDataProvider.h
--- a/Engine/Plugins/Runtime/ComputeFramework/Source/ComputeFramework/Public/ComputeFramework/ComputeDataProvider.h
+++ b/Engine/Plugins/Runtime/ComputeFramework/Source/ComputeFramework/Public/ComputeFramework/ComputeDataProvider.h
@@ -12,6 +12,7 @@ class FComputeDataProviderRenderProxy;
 struct FComputeKernelPermutationVector;
 class FRDGBuilder;
 class FRDGExternalAccessQueue;
+struct FComputeContext;
 
 /**
  * Compute Framework Data Provider.
@@ -122,13 +123,10 @@ public:
 	virtual void GetReadbackData(TArray<FReadbackData>& OutReadbackData) const {}
 
 	/** Called immediately prior to every kernel dispatch that outputs to this data. */
-	virtual void PreSubmit(FRDGBuilder& InGraphBuilder) const {}
+	virtual void PreSubmit(FComputeContext& Context) const {}
 
 	/** Called immediately after every kernel dispatch that outputs to this data. */
-	virtual void PostSubmit(FRDGBuilder& InGraphBuilder) const {}
-
-	UE_DEPRECATED(5.6, "Use PostSubmit instead.")
-	virtual void PostGraphDispatch(FRDGBuilder& GraphBuilder) const {}
+	virtual void PostSubmit(FComputeContext& Context) const {}
 
 protected:
 	/** Helper for making an FStridedView over the FDispatchData. */
@@ -137,4 +135,12 @@ protected:
 	{
 		return TStridedView<ElementType>(InDispatchData.ParameterBufferStride, (ElementType*)(InDispatchData.ParameterBuffer + InDispatchData.ParameterBufferOffset), InDispatchData.NumInvocations);
 	}
+
+public:
+	UE_DEPRECATED(5.6, "Use PostSubmit instead.")
+	virtual void PostGraphDispatch(FRDGBuilder& GraphBuilder) const {}
+	UE_DEPRECATED(5.7, "Use version that takes FComputeContext instead.")
+	COMPUTEFRAMEWORK_API virtual void PreSubmit(FRDGBuilder& InGraphBuilder) const;
+	UE_DEPRECATED(5.7, "Use version that takes FComputeContext instead.")
+	COMPUTEFRAMEWORK_API virtual void PostSubmit(FRDGBuilder& InGraphBuilder) const;
 };
diff --git a/Engine/Source/Runtime/Engine/Private/Commandlets/Commandlet.cpp b/Engine/Source/Runtime/Engine/Private/Commandlets/Commandlet.cpp
--- a/Engine/Source/Runtime/Engine/Private/Commandlets/Commandlet.cpp
+++ b/Engine/Source/Runtime/Engine/Private/Commandlets/Commandlet.cpp
@@ -175,7 +175,7 @@ void CommandletHelpers::TickEngine(UWorld* InWorld, double InDeltaTime)
 			TArray<IComputeTaskWorker*> ComputeWorkers;
 			Scene->GetComputeTaskWorkers(ComputeWorkers);
 
-			ENQUEUE_RENDER_COMMAND(ComputeTaskWorker_EndOfFrameUpdate)([FeatureLevel, ComputeWorkers](FRHICommandListImmediate& RHICmdList)
+			ENQUEUE_RENDER_COMMAND(ComputeTaskWorker_EndOfFrameUpdate)([FeatureLevel, ComputeWorkers, Scene](FRHICommandListImmediate& RHICmdList)
 			{				
 				FRDGBuilder GraphBuilder(RHICmdList, FRDGEventName(TEXT("ComputeTaskWorker_EndOfFrameUpdate")));
 				
@@ -183,7 +183,15 @@ void CommandletHelpers::TickEngine(UWorld* InWorld, double InDeltaTime)
 				{
 					if (ComputeTaskWorker->HasWork(ComputeTaskExecutionGroup::EndOfFrameUpdate))
 					{
-						ComputeTaskWorker->SubmitWork(GraphBuilder, ComputeTaskExecutionGroup::EndOfFrameUpdate, FeatureLevel);
+						FComputeContext Context
+						{
+							.GraphBuilder = GraphBuilder,
+							.ExecutionGroupName = ComputeTaskExecutionGroup::EndOfFrameUpdate,
+							.FeatureLevel = FeatureLevel,
+							.Scene = Scene,
+						};
+
+						ComputeTaskWorker->SubmitWork(Context);
 					}
 				}
 
diff --git a/Engine/Source/Runtime/Engine/Private/ComputeWorkerInterface.cpp b/Engine/Source/Runtime/Engine/Private/ComputeWorkerInterface.cpp
--- a/Engine/Source/Runtime/Engine/Private/ComputeWorkerInterface.cpp
+++ b/Engine/Source/Runtime/Engine/Private/ComputeWorkerInterface.cpp
@@ -3,6 +3,21 @@
 #include "ComputeWorkerInterface.h"
 #include "UObject/NameTypes.h"
 
+void IComputeTaskWorker::SubmitWork(FRDGBuilder& GraphBuilder, FName InExecutionGroupName, ERHIFeatureLevel::Type InFeatureLevel)
+{
+	FComputeContext Context = FComputeContext
+	{
+		.GraphBuilder = GraphBuilder,
+		.ExecutionGroupName = InExecutionGroupName,
+		.FeatureLevel = InFeatureLevel,
+		.Scene = nullptr,
+		.View = nullptr,
+	};
+
+	SubmitWork(Context);
+}
+
 FName ComputeTaskExecutionGroup::Immediate("Immediate");
 FName ComputeTaskExecutionGroup::EndOfFrameUpdate("EndOfFrameUpdate");
 FName ComputeTaskExecutionGroup::BeginInitViews("BeginInitViews");
+FName ComputeTaskExecutionGroup::PostTLASBuild("PostTLASBuild");
diff --git a/Engine/Source/Runtime/Engine/Public/ComputeWorkerInterface.h b/Engine/Source/Runtime/Engine/Public/ComputeWorkerInterface.h
--- a/Engine/Source/Runtime/Engine/Public/ComputeWorkerInterface.h
+++ b/Engine/Source/Runtime/Engine/Public/ComputeWorkerInterface.h
@@ -5,8 +5,19 @@
 #include "CoreMinimal.h"
 
 class FRDGBuilder;
+class FSceneInterface;
+class FSceneView;
 namespace ERHIFeatureLevel { enum Type : int; }
 
+struct FComputeContext
+{
+	FRDGBuilder& GraphBuilder;
+	FName ExecutionGroupName;
+	ERHIFeatureLevel::Type FeatureLevel = static_cast<ERHIFeatureLevel::Type>(0);
+	const FSceneInterface* Scene = nullptr;
+	const FSceneView* View = nullptr; // Null if executed outside of view rendering.
+};
+
 /** 
  * Interface for a compute task worker.
  * Implementations will queue and schedule work per scene before the renderer submits at fixed points in the frame.
@@ -20,7 +31,10 @@ public:
 	virtual bool HasWork(FName InExecutionGroupName) const = 0;
 
 	/** Add any scheduled work to an RDGBuilder ready for execution. */
-	virtual void SubmitWork(FRDGBuilder& GraphBuilder, FName InExecutionGroupName, ERHIFeatureLevel::Type FeatureLevel) = 0;
+	virtual void SubmitWork(FComputeContext& Context) = 0;
+
+	UE_DEPRECATED(5.7, "Use version that takes FComputeContext instead")
+	ENGINE_API virtual void SubmitWork(FRDGBuilder& GraphBuilder, FName InExecutionGroupName, ERHIFeatureLevel::Type InFeatureLevel);
 };
 
 /** Core execution group names for use in IComputeTaskWorker::SubmitWork(). */
@@ -29,4 +43,5 @@ struct ComputeTaskExecutionGroup
 	static ENGINE_API FName Immediate;
 	static ENGINE_API FName EndOfFrameUpdate;
 	static ENGINE_API FName BeginInitViews;
+	static ENGINE_API FName PostTLASBuild;
 };
diff --git a/Engine/Source/Runtime/Engine/Public/SceneViewExtension.h b/Engine/Source/Runtime/Engine/Public/SceneViewExtension.h
--- a/Engine/Source/Runtime/Engine/Public/SceneViewExtension.h
+++ b/Engine/Source/Runtime/Engine/Public/SceneViewExtension.h
@@ -100,6 +100,15 @@ using FPostProcessingPassDelegateArray = TArray<FPostProcessingPassDelegate, Sce
 using FAfterPassCallbackDelegate = FPostProcessingPassDelegate;
 using FAfterPassCallbackDelegateArray = FPostProcessingPassDelegateArray;
 
+/** Misc flags for communicating requirements or state externally. */
+enum class ESceneViewExtensionFlags : uint32
+{
+	None = 0,
+
+	SubscribesToPostTLASBuild = 1 << 0
+};
+ENUM_CLASS_FLAGS(ESceneViewExtensionFlags);
+
 class ISceneViewExtension
 {
 public:
@@ -185,6 +194,11 @@ public:
 	 */
 	virtual void PostRenderBasePassMobile_RenderThread(FRHICommandList& RHICmdList, FSceneView& InView) {}
 
+	/**
+	 * Called after all the TLAS has been built.
+	 */
+	virtual void PostTLASBuild_RenderThread(FRDGBuilder& GraphBuilder, FSceneView& InView) {}
+
 	/**
 	 * Called right before Post Processing rendering begins
 	 */
@@ -222,6 +236,11 @@ public:
 	 */
 	virtual bool IsActiveThisFrame(const FSceneViewExtensionContext& Context) const { return IsActiveThisFrame_Internal(Context); }
 
+	/**
+	 * Query any flags to communicate state or execution requirements externally.
+	 */
+	virtual ESceneViewExtensionFlags GetFlags() const { return ESceneViewExtensionFlags::None; }
+
 	///////////////////////////////////////////////////////////////////////////////////////
 	//! Deprecated APIs - These are no longer called and must be converted to restore functionality.
 
diff --git a/Engine/Source/Runtime/Renderer/Private/DeferredShadingRenderer.cpp b/Engine/Source/Runtime/Renderer/Private/DeferredShadingRenderer.cpp
--- a/Engine/Source/Runtime/Renderer/Private/DeferredShadingRenderer.cpp
+++ b/Engine/Source/Runtime/Renderer/Private/DeferredShadingRenderer.cpp
@@ -3001,6 +3001,28 @@ void FDeferredShadingSceneRenderer::Render(FRDGBuilder& GraphBuilder, const FSce
 
 		// End shadow and fog after base pass
 
+#if RHI_RAYTRACING
+		if (IsRayTracingEnabled(ViewFamily.GetShaderPlatform()) && GRHISupportsRayTracingShaders)
+		{
+			for (int32 ViewExt = 0; ViewExt < ViewFamily.ViewExtensions.Num(); ++ViewExt)
+			{
+				if (EnumHasAnyFlags(ViewFamily.ViewExtensions[ViewExt]->GetFlags(), ESceneViewExtensionFlags::SubscribesToPostTLASBuild))
+				{
+					if (bNeedToSetupRayTracingRenderingData)
+					{
+						SetupRayTracingRenderingData(GraphBuilder);
+						bNeedToSetupRayTracingRenderingData = false;
+					}
+
+					for (int32 ViewIndex = 0; ViewIndex < ViewFamily.Views.Num(); ++ViewIndex)
+					{
+						ViewFamily.ViewExtensions[ViewExt]->PostTLASBuild_RenderThread(GraphBuilder, Views[ViewIndex]);
+					}
+				}
+			}
+		}
+#endif // RHI_RAYTRACING
+
 		if (bNaniteEnabled)
 		{
 			// Needs doing after shadows such that the checks for shadow atlases etc work.
diff --git a/Engine/Source/Runtime/Renderer/Private/RendererScene.cpp b/Engine/Source/Runtime/Renderer/Private/RendererScene.cpp
--- a/Engine/Source/Runtime/Renderer/Private/RendererScene.cpp
+++ b/Engine/Source/Runtime/Renderer/Private/RendererScene.cpp
@@ -6535,7 +6535,15 @@ void FScene::Update(FRDGBuilder& GraphBuilder, const FUpdateParameters& Paramete
 	{
 		if (ComputeTaskWorker->HasWork(ComputeTaskExecutionGroup::EndOfFrameUpdate))
 		{
-			ComputeTaskWorker->SubmitWork(GraphBuilder, ComputeTaskExecutionGroup::EndOfFrameUpdate, FeatureLevel);
+			FComputeContext Context
+			{
+				.GraphBuilder = GraphBuilder,
+				.ExecutionGroupName = ComputeTaskExecutionGroup::EndOfFrameUpdate,
+				.FeatureLevel = FeatureLevel,
+				.Scene = this,
+			};
+
+			ComputeTaskWorker->SubmitWork(Context);
 		}
 	}
 
diff --git a/Engine/Source/Runtime/Renderer/Private/SceneVisibility.cpp b/Engine/Source/Runtime/Renderer/Private/SceneVisibility.cpp
--- a/Engine/Source/Runtime/Renderer/Private/SceneVisibility.cpp
+++ b/Engine/Source/Runtime/Renderer/Private/SceneVisibility.cpp
@@ -4966,13 +4966,21 @@ void FSceneRenderer::PreVisibilityFrameSetup(FRDGBuilder& GraphBuilder)
 			}
 		}
 
-		if(Scene)
+		if (Scene)
 		{
 			for (IComputeTaskWorker* ComputeTaskWorker : Scene->ComputeTaskWorkers)
 			{
 				if (ComputeTaskWorker->HasWork(ComputeTaskExecutionGroup::BeginInitViews))
 				{
-					ComputeTaskWorker->SubmitWork(GraphBuilder, ComputeTaskExecutionGroup::BeginInitViews, FeatureLevel);
+					FComputeContext Context
+					{
+						.GraphBuilder = GraphBuilder,
+						.ExecutionGroupName = ComputeTaskExecutionGroup::BeginInitViews,
+						.FeatureLevel = FeatureLevel,
+						.Scene = Scene,
+					};
+
+					ComputeTaskWorker->SubmitWork(Context);
 				}
 			}
 		}
