[ComputeFramework] Minor tidy up of new VE

#rnx

diff --git a/Engine/Plugins/Runtime/ComputeFramework/Source/ComputeFramework/Private/ComputeFramework/ComputeViewExtension.cpp b/Engine/Plugins/Runtime/ComputeFramework/Source/ComputeFramework/Private/ComputeFramework/ComputeViewExtension.cpp
--- a/Engine/Plugins/Runtime/ComputeFramework/Source/ComputeFramework/Private/ComputeFramework/ComputeViewExtension.cpp
+++ b/Engine/Plugins/Runtime/ComputeFramework/Source/ComputeFramework/Private/ComputeFramework/ComputeViewExtension.cpp
@@ -11,7 +11,6 @@ FComputeViewExtension::FComputeViewExtension(const FAutoRegister& AutoReg, UWorl
 
 void FComputeViewExtension::PostTLASBuild_RenderThread(FRDGBuilder& GraphBuilder, FSceneView& InView)
 {
-	// Verified this gets called, if a graph that requires it sets it up. If we're happy with these changes, I can add a new PostTLASBuild function to the view extension interface?
 	const TWeakObjectPtr<UWorld> WorldPtr = GetWorld();
 	if (!WorldPtr.IsValid())
 	{
@@ -51,7 +50,7 @@ ESceneViewExtensionFlags FComputeViewExtension::GetFlags() const
 {
 	const TWeakObjectPtr<UWorld> WorldPtr = GetWorld();
 
-	if (const FSceneInterface* Scene = WorldPtr.IsValid() ? WorldPtr.Get()->Scene : nullptr)
+	if (const FSceneInterface* Scene = WorldPtr.IsValid() ? WorldPtr->Scene : nullptr)
 	{
 		ComputeWorkers.Empty(ComputeWorkers.Num());
 		Scene->GetComputeTaskWorkers(ComputeWorkers);
