[ComputeFramework] Fix issue where CF system view extension gets overwritten every time a new scene is created

#rnx

diff --git a/Engine/Plugins/Runtime/ComputeFramework/Source/ComputeFramework/Private/ComputeFramework/ComputeSystem.h b/Engine/Plugins/Runtime/ComputeFramework/Source/ComputeFramework/Private/ComputeFramework/ComputeSystem.h
--- a/Engine/Plugins/Runtime/ComputeFramework/Source/ComputeFramework/Private/ComputeFramework/ComputeSystem.h
+++ b/Engine/Plugins/Runtime/ComputeFramework/Source/ComputeFramework/Private/ComputeFramework/ComputeSystem.h
@@ -21,6 +21,5 @@ public:
 
 private:
 	TMap<FSceneInterface const*, FComputeGraphTaskWorker*> ComputeWorkers;
-
-	TSharedPtr<FComputeViewExtension, ESPMode::ThreadSafe> ViewExtension;
+	TMap<FSceneInterface const*, TSharedPtr<FComputeViewExtension, ESPMode::ThreadSafe>> ViewExtensions;
 };
diff --git a/Engine/Plugins/Runtime/ComputeFramework/Source/ComputeFramework/Private/ComputeSystem.cpp b/Engine/Plugins/Runtime/ComputeFramework/Source/ComputeFramework/Private/ComputeSystem.cpp
--- a/Engine/Plugins/Runtime/ComputeFramework/Source/ComputeFramework/Private/ComputeSystem.cpp
+++ b/Engine/Plugins/Runtime/ComputeFramework/Source/ComputeFramework/Private/ComputeSystem.cpp
@@ -13,7 +13,7 @@ void FComputeFrameworkSystem::CreateWorkers(FSceneInterface const* InScene, TArr
 
 	if (UWorld* World = InScene->GetWorld())
 	{
-		ViewExtension = FSceneViewExtensions::NewExtension<FComputeViewExtension>(World);
+		ViewExtensions.Add(InScene, FSceneViewExtensions::NewExtension<FComputeViewExtension>(World));
 	}
 }
 
@@ -31,7 +31,7 @@ void FComputeFrameworkSystem::DestroyWorkers(FSceneInterface const* InScene, TAr
 		ComputeWorkers.Remove(InScene);
 	}
 
-	ViewExtension.Reset();
+	ViewExtensions.Remove(InScene);
 }
 
 FComputeGraphTaskWorker* FComputeFrameworkSystem::GetComputeWorker(FSceneInterface const* InScene) const
