[ComputeFramework] Mark UComputeKernelSource properties only required for compilation as editor only

#rnx

#Synced-CL 51167808

diff --git a/Engine/Plugins/Animation/DeformerGraph/Source/OptimusCore/Private/Nodes/OptimusNode_ComputeKernelBase.cpp b/Engine/Plugins/Animation/DeformerGraph/Source/OptimusCore/Private/Nodes/OptimusNode_ComputeKernelBase.cpp
--- a/Engine/Plugins/Animation/DeformerGraph/Source/OptimusCore/Private/Nodes/OptimusNode_ComputeKernelBase.cpp
+++ b/Engine/Plugins/Animation/DeformerGraph/Source/OptimusCore/Private/Nodes/OptimusNode_ComputeKernelBase.cpp
@@ -54,6 +54,7 @@ static void CopyValueType(FShaderValueTypeHandle InValueType,  FShaderParamTypeD
 	OutParamDef.ResetTypeDeclaration();
 }
 
+#if WITH_EDITOR
 FOptimus_ComputeKernelResult UOptimusNode_ComputeKernelBase::CreateComputeKernel(
 	UObject *InKernelSourceOuter,
 	const FOptimusPinTraversalContext& InTraversalContext,
@@ -262,7 +263,7 @@ FOptimus_ComputeKernelResult UOptimusNode_ComputeKernelBase::CreateComputeKernel
 	
 	return FOptimus_ComputeKernelResult(TInPlaceType<UOptimusKernelSource*>(), KernelSource);
 }
-
+#endif // WITH_EDITOR
 
 TOptional<FText> UOptimusNode_ComputeKernelBase::ValidateForCompile(const FOptimusPinTraversalContext& InContext) const
 {
@@ -391,7 +392,7 @@ const TCHAR* UOptimusNode_ComputeKernelBase::GetNumThreadsReservedName()
 	return NumThreadsReservedName;
 }
 
-
+#if WITH_EDITOR
 TOptional<FText> UOptimusNode_ComputeKernelBase::ProcessInputPinForComputeKernel(
 	bool bSupportUnifiedDispatchOnly,
 	const FOptimusPinTraversalContext& InTraversalContext,
@@ -816,5 +817,6 @@ TOptional<FText> UOptimusNode_ComputeKernelBase::BindKernelDataInterfaceForCompu
 
 	return {};
 }
+#endif // WITH_EDITOR
 
 #undef LOCTEXT_NAMESPACE
diff --git a/Engine/Plugins/Animation/DeformerGraph/Source/OptimusCore/Private/Nodes/OptimusNode_ComputeKernelBase.h b/Engine/Plugins/Animation/DeformerGraph/Source/OptimusCore/Private/Nodes/OptimusNode_ComputeKernelBase.h
--- a/Engine/Plugins/Animation/DeformerGraph/Source/OptimusCore/Private/Nodes/OptimusNode_ComputeKernelBase.h
+++ b/Engine/Plugins/Animation/DeformerGraph/Source/OptimusCore/Private/Nodes/OptimusNode_ComputeKernelBase.h
@@ -35,6 +35,7 @@ public:
 	/** Implement this to return additional source code that the kernel requires. */
 	virtual TArray<TObjectPtr<UComputeSource>> GetAdditionalSources() const { return {}; }
 	
+#if WITH_EDITOR
 	// IOptimusComputeKernelProvider
 	FOptimus_ComputeKernelResult CreateComputeKernel(
 		UObject* InKernelSourceOuter,
@@ -45,6 +46,7 @@ public:
 		FOptimus_InterfaceBindingMap& OutInputDataBindings,
 		FOptimus_InterfaceBindingMap& OutOutputDataBindings
 	) const override;
+#endif
 
 	FOptimusExecutionDomain GetExecutionDomain() const override PURE_VIRTUAL(UOptimusNode_ComputeKernelBase::GetExecutionDomain, return {}; );
 	const UOptimusNodePin* GetPrimaryGroupPin() const override PURE_VIRTUAL(UOptimusNode_ComputeKernelBase::GetPrimaryGroupPin, return {}; );  
@@ -96,6 +98,7 @@ protected:
 	static FString GetReadFunctionName(const FString& InBindingName);
 	static const TCHAR* GetNumThreadsReservedName();
 	
+#if WITH_EDITOR
 private:
 	TOptional<FText> BindKernelDataInterfaceForComputeKernel(
 		const UOptimusComponentSourceBinding* InKernelPrimaryComponentSourceBinding,
@@ -129,6 +132,6 @@ private:
 		FOptimus_InterfaceBindingMap& OutOutputDataBindings,
 		bool& bOutIsUnifiedDispatch
 		) const;
-
+#endif // WITH_EDITOR
 };
 
diff --git a/Engine/Plugins/Animation/DeformerGraph/Source/OptimusCore/Private/OptimusDeformer.cpp b/Engine/Plugins/Animation/DeformerGraph/Source/OptimusCore/Private/OptimusDeformer.cpp
--- a/Engine/Plugins/Animation/DeformerGraph/Source/OptimusCore/Private/OptimusDeformer.cpp
+++ b/Engine/Plugins/Animation/DeformerGraph/Source/OptimusCore/Private/OptimusDeformer.cpp
@@ -1494,7 +1494,7 @@ static void CollectNodes(
 	}	
 }
 
-
+#if WITH_EDITOR
 bool UOptimusDeformer::Compile()
 {
 	if (!GetUpdateGraph())
@@ -1605,6 +1605,7 @@ bool UOptimusDeformer::Compile()
 	
 	return true;
 }
+#endif // WITH_EDITOR
 
 TArray<UOptimusNode*> UOptimusDeformer::GetAllNodesOfClass(UClass* InNodeClass) const
 {
@@ -1718,6 +1719,7 @@ struct FOptimusInstancedPin
 	}
 };
 
+#if WITH_EDITOR
 FOptimusNodeGraphCompilationResult UOptimusDeformer::CompileNodeGraphToComputeGraphs(
 	const UOptimusNodeGraph* InNodeGraph,
 	TFunction<void(EOptimusDiagnosticLevel, FText, const UObject*)> InErrorReporter
@@ -3208,6 +3210,7 @@ FOptimusNodeGraphCompilationResult UOptimusDeformer::CompileNodeGraphToComputeGr
 
 	return Result;
 }
+#endif // WITH_EDITOR
 
 void UOptimusDeformer::OnDataTypeChanged(FName InTypeName)
 {
@@ -3504,6 +3507,7 @@ void UOptimusDeformer::PostLoad()
 		}
 	}
 
+#if WITH_EDITOR
 	FOptimusObjectVersion::Type ObjectVersion = static_cast<FOptimusObjectVersion::Type>(GetLinkerCustomVersion(FOptimusObjectVersion::GUID));
 	if (ObjectVersion < FOptimusObjectVersion::ReparentResourcesAndVariables)
 	{
@@ -3620,6 +3624,7 @@ void UOptimusDeformer::PostLoad()
 			Compile();
 		}
 	}
+#endif // WITH_EDITOR
 }
 
 
diff --git a/Engine/Plugins/Animation/DeformerGraph/Source/OptimusCore/Public/IOptimusComputeKernelProvider.h b/Engine/Plugins/Animation/DeformerGraph/Source/OptimusCore/Public/IOptimusComputeKernelProvider.h
--- a/Engine/Plugins/Animation/DeformerGraph/Source/OptimusCore/Public/IOptimusComputeKernelProvider.h
+++ b/Engine/Plugins/Animation/DeformerGraph/Source/OptimusCore/Public/IOptimusComputeKernelProvider.h
@@ -71,6 +71,7 @@ class IOptimusComputeKernelProvider
 	GENERATED_BODY()
 
 public:
+#if WITH_EDITOR
 	/**
 	 * Return an UOptimusKernelSource object, from a compute kernel node state that implements
 	 * this interface.
@@ -89,6 +90,7 @@ public:
 		FOptimus_InterfaceBindingMap& OutInputDataBindings,
 		FOptimus_InterfaceBindingMap& OutOutputDataBindings
 	) const = 0;
+#endif
 
 	/** Returns the execution domain that this kernel should iterate over */
 	virtual FOptimusExecutionDomain GetExecutionDomain() const = 0;
diff --git a/Engine/Plugins/Animation/DeformerGraph/Source/OptimusCore/Public/OptimusDeformer.h b/Engine/Plugins/Animation/DeformerGraph/Source/OptimusCore/Public/OptimusDeformer.h
--- a/Engine/Plugins/Animation/DeformerGraph/Source/OptimusCore/Public/OptimusDeformer.h
+++ b/Engine/Plugins/Animation/DeformerGraph/Source/OptimusCore/Public/OptimusDeformer.h
@@ -300,7 +300,9 @@ public:
 
 	
 	/// Graph compilation
+#if WITH_EDITOR
 	UE_API bool Compile();
+#endif
 
 	/** Returns a multicast delegate that can be subscribed to listen for the start of compilation. */
 	FOptimusCompileBegin& GetCompileBeginDelegate()  { return CompileBeginDelegate; }
@@ -542,11 +544,13 @@ private:
 	UE_API void OnGraphRenamedOrRemoved(UOptimusNodeGraph* InGraph) const;
 	UE_API void UpdateFunctionReferenceNodeDisplayName(const FOptimusFunctionGraphIdentifier& InRenamedFunction);
 	
+#if WITH_EDITOR
 	/// Compile a node graph to a compute graph. Returns one or two complete compute graphs if compilation succeeded. 
 	UE_API FOptimusNodeGraphCompilationResult CompileNodeGraphToComputeGraphs(
 		const UOptimusNodeGraph *InNodeGraph,
 		TFunction<void(EOptimusDiagnosticLevel, FText, const UObject*)> InErrorReporter
 		);
+#endif
 
 	UE_API void OnDataTypeChanged(FName InTypeName);
 
diff --git a/Engine/Plugins/Animation/DeformerGraph/Source/OptimusCore/Public/OptimusKernelSource.h b/Engine/Plugins/Animation/DeformerGraph/Source/OptimusCore/Public/OptimusKernelSource.h
--- a/Engine/Plugins/Animation/DeformerGraph/Source/OptimusCore/Public/OptimusKernelSource.h
+++ b/Engine/Plugins/Animation/DeformerGraph/Source/OptimusCore/Public/OptimusKernelSource.h
@@ -11,6 +11,7 @@ class UOptimusKernelSource :
 {
 	GENERATED_BODY()
 
+#if WITH_EDITOR
 public:
 	void SetSource(FString const& InSource)
 	{
@@ -23,8 +24,11 @@ public:
 		return Source;
 	}
 	//~ End UComputeKernelSource Interface.
+#endif
 
+#if WITH_EDITORONLY_DATA
 protected:
 	UPROPERTY()
 	FString Source;
+#endif
 };
diff --git a/Engine/Plugins/PCG/Source/PCG/Private/Compute/PCGComputeKernelSource.h b/Engine/Plugins/PCG/Source/PCG/Private/Compute/PCGComputeKernelSource.h
--- a/Engine/Plugins/PCG/Source/PCG/Private/Compute/PCGComputeKernelSource.h
+++ b/Engine/Plugins/PCG/Source/PCG/Private/Compute/PCGComputeKernelSource.h
@@ -11,20 +11,15 @@ class UPCGComputeKernelSource : public UComputeKernelSource
 {
 	GENERATED_BODY()
 
+#if WITH_EDITOR
 public:
 	//~ Begin UComputeKernelSource Interface.
 	FString GetSource() const override
 	{
-#if WITH_EDITOR
 		return Source;
-#else
-		ensure(false);
-		return "";
-#endif
 	}
 	//~ End UComputeKernelSource Interface.
 
-#if WITH_EDITOR
 	void SetSource(FString const& InSource) { Source = InSource; }
 
 protected:
diff --git a/Engine/Plugins/Runtime/ComputeFramework/Source/ComputeFramework/Public/ComputeFramework/ComputeKernelSource.h b/Engine/Plugins/Runtime/ComputeFramework/Source/ComputeFramework/Public/ComputeFramework/ComputeKernelSource.h
--- a/Engine/Plugins/Runtime/ComputeFramework/Source/ComputeFramework/Public/ComputeFramework/ComputeKernelSource.h
+++ b/Engine/Plugins/Runtime/ComputeFramework/Source/ComputeFramework/Public/ComputeFramework/ComputeKernelSource.h
@@ -20,8 +20,10 @@ class UComputeKernelSource : public UObject
 	GENERATED_BODY()
 
 public:
+#if WITH_EDITOR
 	/** Get kernel source code ready for HLSL compilation. */
 	virtual FString GetSource() const PURE_VIRTUAL(UComputeKernelSource::GetSource, return {};);
+#endif
 
 	/** Kernel entry point. */
 	UPROPERTY(VisibleAnywhere, AssetRegistrySearchable, Category = "Kernel")
@@ -31,6 +33,7 @@ public:
 	UPROPERTY(VisibleAnywhere, AssetRegistrySearchable, Category = "Kernel")
 	FIntVector GroupSize = FIntVector(64, 1, 1);
 
+#if WITH_EDITORONLY_DATA
 	/** Base permutations exposed by the kernel. These will be extended by further permutations declared in any linked data providers. */
 	UPROPERTY(VisibleAnywhere, meta = (ShowOnlyInnerProperties), Category = "Kernel")
 	FComputeKernelPermutationSet PermutationSet;
@@ -50,6 +53,7 @@ public:
 	/* Named external outputs for the kernel. These must be fulfilled by linked data providers. */
 	UPROPERTY(VisibleAnywhere, EditFixedSize, Category = "External")
 	TArray<FShaderFunctionDefinition> ExternalOutputs;
+#endif
 };
 
 /** Simple concrete implementation of UComputeKernelSource with contained source text property. */
@@ -59,13 +63,17 @@ class UComputeKernelSourceWithText final : public UComputeKernelSource
 	GENERATED_BODY()
 
 public:
+#if WITH_EDITORONLY_DATA
 	/* Kernel source text. */
 	UPROPERTY(EditAnywhere, Category = "Kernel")
 	FString SourceText;
+#endif
 
+#if WITH_EDITOR
 	//~ Begin UComputeKernelSource Interface.
 	virtual FString GetSource() const override { return SourceText; }
 	//~ End UComputeKernelSource Interface.
+#endif
 };
 
 #undef UE_API
