ComputeFramework: Removed unused asset types that were only for testing.
#rnx
#Synced-CL 52056483

diff --git a/Engine/Plugins/Runtime/ComputeFramework/Source/ComputeFramework/Private/ComputeGraphFromText.cpp b/Engine/Plugins/Runtime/ComputeFramework/Source/ComputeFramework/Private/ComputeGraphFromText.cpp
deleted file mode 100644
diff --git a/Engine/Plugins/Runtime/ComputeFramework/Source/ComputeFramework/Public/ComputeFramework/ComputeGraphFromText.h b/Engine/Plugins/Runtime/ComputeFramework/Source/ComputeFramework/Public/ComputeFramework/ComputeGraphFromText.h
deleted file mode 100644
diff --git a/Engine/Plugins/Runtime/ComputeFramework/Source/ComputeFrameworkEditor/ComputeFrameworkEditor.Build.cs b/Engine/Plugins/Runtime/ComputeFramework/Source/ComputeFrameworkEditor/ComputeFrameworkEditor.Build.cs
--- a/Engine/Plugins/Runtime/ComputeFramework/Source/ComputeFrameworkEditor/ComputeFrameworkEditor.Build.cs
+++ b/Engine/Plugins/Runtime/ComputeFramework/Source/ComputeFrameworkEditor/ComputeFrameworkEditor.Build.cs
@@ -9,10 +9,8 @@ namespace UnrealBuildTool.Rules
             PrivateDependencyModuleNames.AddRange(
                 new string[]
                 {
-					"AssetTools",
 					"ComputeFramework",
 					"Core",
-					"CoreUObject",
 					"Engine",
 					"UnrealEd",
 				}
diff --git a/Engine/Plugins/Runtime/ComputeFramework/Source/ComputeFrameworkEditor/Private/ComputeFramework/ComputeFrameworkEditorModule.h b/Engine/Plugins/Runtime/ComputeFramework/Source/ComputeFrameworkEditor/Private/ComputeFramework/ComputeFrameworkEditorModule.h
--- a/Engine/Plugins/Runtime/ComputeFramework/Source/ComputeFrameworkEditor/Private/ComputeFramework/ComputeFrameworkEditorModule.h
+++ b/Engine/Plugins/Runtime/ComputeFramework/Source/ComputeFrameworkEditor/Private/ComputeFramework/ComputeFrameworkEditorModule.h
@@ -3,7 +3,6 @@
 #pragma once
 
 #include "ComputeFramework/IComputeFrameworkEditorModule.h"
-#include "ComputeFramework/ComputeFrameworkCompilationTick.h"
 
 class FComputeFrameworkEditorModule : public IComputeFrameworkEditorModule
 {
diff --git a/Engine/Plugins/Runtime/ComputeFramework/Source/ComputeFrameworkEditor/Private/ComputeFramework/ComputeGraphFromTextAssetActions.h b/Engine/Plugins/Runtime/ComputeFramework/Source/ComputeFrameworkEditor/Private/ComputeFramework/ComputeGraphFromTextAssetActions.h
deleted file mode 100644
diff --git a/Engine/Plugins/Runtime/ComputeFramework/Source/ComputeFrameworkEditor/Private/ComputeFramework/ComputeGraphFromTextFactory.h b/Engine/Plugins/Runtime/ComputeFramework/Source/ComputeFrameworkEditor/Private/ComputeFramework/ComputeGraphFromTextFactory.h
deleted file mode 100644
diff --git a/Engine/Plugins/Runtime/ComputeFramework/Source/ComputeFrameworkEditor/Private/ComputeFrameworkEditorModule.cpp b/Engine/Plugins/Runtime/ComputeFramework/Source/ComputeFrameworkEditor/Private/ComputeFrameworkEditorModule.cpp
--- a/Engine/Plugins/Runtime/ComputeFramework/Source/ComputeFrameworkEditor/Private/ComputeFrameworkEditorModule.cpp
+++ b/Engine/Plugins/Runtime/ComputeFramework/Source/ComputeFrameworkEditor/Private/ComputeFrameworkEditorModule.cpp
@@ -2,17 +2,12 @@
 
 #include "ComputeFramework/ComputeFrameworkEditorModule.h"
 
-#include "AssetToolsModule.h"
 #include "ComputeFramework/ComputeFrameworkCompilationTick.h"
-#include "ComputeFramework/ComputeGraphFromTextAssetActions.h"
 #include "Modules/ModuleManager.h"
 
 void FComputeFrameworkEditorModule::StartupModule()
 {
 	TickObject = MakeUnique<FComputeFrameworkCompilationTick>();
-
- 	IAssetTools& AssetTools = FModuleManager::LoadModuleChecked<FAssetToolsModule>("AssetTools").Get();
- 	AssetTools.RegisterAssetTypeActions(MakeShareable(new FAssetTypeActions_ComputeGraphFromText));
 }
 
 void FComputeFrameworkEditorModule::ShutdownModule()
diff --git a/Engine/Plugins/Runtime/ComputeFramework/Source/ComputeFrameworkEditor/Private/ComputeGraphFromTextAssetActions.cpp b/Engine/Plugins/Runtime/ComputeFramework/Source/ComputeFrameworkEditor/Private/ComputeGraphFromTextAssetActions.cpp
deleted file mode 100644
diff --git a/Engine/Plugins/Runtime/ComputeFramework/Source/ComputeFrameworkEditor/Private/ComputeGraphFromTextFactory.cpp b/Engine/Plugins/Runtime/ComputeFramework/Source/ComputeFrameworkEditor/Private/ComputeGraphFromTextFactory.cpp
deleted file mode 100644
