ComputeFramework: Make dispatch and buffer data interfaces configurable.
Modify ComputeGraphComponent to support setting exposed properties on the data interfaces.
Move data interfaces to separate module.
#rnx
#Synced-CL 52056483

diff --git a/Engine/Plugins/Runtime/ComputeFramework/ComputeFramework.uplugin b/Engine/Plugins/Runtime/ComputeFramework/ComputeFramework.uplugin
--- a/Engine/Plugins/Runtime/ComputeFramework/ComputeFramework.uplugin
+++ b/Engine/Plugins/Runtime/ComputeFramework/ComputeFramework.uplugin
@@ -24,6 +24,11 @@
 			"Name": "ComputeFrameworkEditor",
 			"Type": "Editor",
 			"LoadingPhase": "Default"
+		},
+		{
+			"Name": "ComputeDataInterface",
+			"Type": "Runtime",
+			"LoadingPhase": "Default"
 		}
 	]
 }
diff --git a/Engine/Plugins/Runtime/ComputeFramework/Source/ComputeDataInterface/ComputeDataInterface.Build.cs b/Engine/Plugins/Runtime/ComputeFramework/Source/ComputeDataInterface/ComputeDataInterface.Build.cs
new file mode 100644
--- /dev/null
+++ b/Engine/Plugins/Runtime/ComputeFramework/Source/ComputeDataInterface/ComputeDataInterface.Build.cs
@@ -0,0 +1,28 @@
+// Copyright Epic Games, Inc. All Rights Reserved.
+
+namespace UnrealBuildTool.Rules
+{
+	public class ComputeDataInterface : ModuleRules
+	{
+		public ComputeDataInterface(ReadOnlyTargetRules Target) : base(Target)
+		{
+			PublicDependencyModuleNames.AddRange(
+				new string[]
+				{
+					"Core",
+					"CoreUObject",
+					"Engine",
+					"ComputeFramework",
+				}
+			);
+
+			PrivateDependencyModuleNames.AddRange(
+				new string[]
+				{
+					"RenderCore",
+					"RHI",
+				}
+			);
+		}
+	}
+}
diff --git a/Engine/Plugins/Runtime/ComputeFramework/Source/ComputeDataInterface/Internal/ComputeFramework/ComputeDataInterfaceBuffer.h b/Engine/Plugins/Runtime/ComputeFramework/Source/ComputeDataInterface/Internal/ComputeFramework/ComputeDataInterfaceBuffer.h
new file mode 100644
--- /dev/null
+++ b/Engine/Plugins/Runtime/ComputeFramework/Source/ComputeDataInterface/Internal/ComputeFramework/ComputeDataInterfaceBuffer.h
@@ -0,0 +1,107 @@
+// Copyright Epic Games, Inc. All Rights Reserved.
+
+#pragma once
+
+#include "ComputeFramework/ComputeDataInterface.h"
+#include "ComputeFramework/ComputeDataProvider.h"
+#include "ComputeFramework/ShaderParamTypeDefinition.h"
+
+#include "ComputeDataInterfaceBuffer.generated.h"
+
+class FBufferDataInterfaceParameters;
+
+UENUM()
+enum class EComputeDataInterfaceBufferType : uint8
+{
+	Int,
+	Uint,
+	Float,
+	Int2,
+	Uint2,
+	Float2,
+	Int3,
+	Uint3,
+	Float3,
+	Int4,
+	Uint4,
+	Float4,
+};
+
+/** Compute data interface used to own and give access to a GPU buffer. */
+UCLASS(MinimalAPI, Category = ComputeFramework)
+class UComputeDataInterfaceBuffer : public UComputeDataInterface
+{
+	GENERATED_BODY()
+
+	//~ Begin UComputeDataInterface Interface
+	TCHAR const* GetClassName() const override { return TEXT("Buffer"); }
+	void GetSupportedInputs(TArray<FShaderFunctionDefinition>& OutFunctions) const override;
+	void GetSupportedOutputs(TArray<FShaderFunctionDefinition>& OutFunctions) const override;
+	void GetShaderParameters(TCHAR const* UID, FShaderParametersMetadataBuilder& InOutBuilder, FShaderParametersMetadataAllocations& InOutAllocations) const override;
+	TCHAR const* GetShaderVirtualPath() const override;
+	void GetShaderHash(FString& InOutKey) const override;
+	void GetHLSL(FString& OutHLSL, FString const& InDataInterfaceName) const override;
+	UComputeDataProvider* CreateDataProvider() const override;
+	//~ End UComputeDataInterface Interface
+
+public:
+	/** The value type for the buffer. */
+	UPROPERTY(EditAnywhere, Category = Settings)
+	EComputeDataInterfaceBufferType ValueType;
+
+	/** Whether to allow read/write access. */
+	UPROPERTY(EditAnywhere, Category = Settings)
+	bool bAllowReadWrite = false;
+
+private:
+	static TCHAR const* TemplateFilePath;
+};
+
+/** Compute data provider implementation for UComputeDataInterfaceBuffer. */
+UCLASS(MinimalAPI, Blueprintable, Category = ComputeFramework)
+class UBufferDataProvider : public UComputeDataProvider
+{
+	GENERATED_BODY()
+
+	//~ Begin UComputeDataProvider Interface
+	void Initialize(int32 InDataInterfaceIndex, UComputeDataInterface const* InDataInterface, UObject* InBinding, uint64 InInputMask, uint64 InOutputMask) override;
+	FComputeDataProviderRenderProxy* GetRenderProxy() override;
+	//~ End UComputeDataProvider Interface
+
+public:
+	/** The value type for the buffer. */
+	UPROPERTY()
+	EComputeDataInterfaceBufferType ValueType;
+
+	/** The size of the buffer. */
+	UPROPERTY(EditAnywhere, BlueprintReadWrite, Category = DataInterface)
+	int32 ElementCount = 0;
+
+	/** Whether to clear the buffer before use. */
+	UPROPERTY(EditAnywhere, BlueprintReadWrite, Category = DataInterface)
+	bool bClearBeforeUse = false;
+};
+
+/** Compute data provider proxy implementation for UComputeDataInterfaceBuffer. */
+class FBufferDataProviderProxy : public FComputeDataProviderRenderProxy
+{
+	//~ Begin FComputeDataProviderRenderProxy Interface
+	bool IsValid(FValidationData const& InValidationData) const override;
+	void AllocateResources(FRDGBuilder& GraphBuilder, FAllocationData const& InAllocationData) override;
+	void GatherDispatchData(FDispatchData const& InDispatchData) override;
+	//~ End FComputeDataProviderRenderProxy Interface
+
+public:
+	FBufferDataProviderProxy(uint32 InElementStride, uint32 InElementCount, bool bInClearBeforeUse);
+
+private:
+	using FParameters = FBufferDataInterfaceParameters;
+
+	const uint32 ElementStride;
+	const uint32 ElementCount;
+	const bool bClearBeforeUse;
+
+	FRDGBufferRef Buffer;
+	FRDGBufferSRVRef BufferSRV;
+	FRDGBufferUAVRef BufferUAV;
+};
diff --git a/Engine/Plugins/Runtime/ComputeFramework/Source/ComputeDataInterface/Internal/ComputeFramework/ComputeDataInterfaceDispatch.h b/Engine/Plugins/Runtime/ComputeFramework/Source/ComputeDataInterface/Internal/ComputeFramework/ComputeDataInterfaceDispatch.h
new file mode 100644
--- /dev/null
+++ b/Engine/Plugins/Runtime/ComputeFramework/Source/ComputeDataInterface/Internal/ComputeFramework/ComputeDataInterfaceDispatch.h
@@ -0,0 +1,65 @@
+// Copyright Epic Games, Inc. All Rights Reserved.
+
+#pragma once
+
+#include "ComputeFramework/ComputeDataInterface.h"
+#include "ComputeFramework/ComputeDataProvider.h"
+
+#include "ComputeDataInterfaceDispatch.generated.h"
+
+class FDataInterfaceDispatchParameters;
+
+/** Compute data interface used to control a static dispatch size, and expose thread count to the shader. */
+UCLASS(MinimalAPI, Category = ComputeFramework)
+class UComputeDataInterfaceDispatch : public UComputeDataInterface
+{
+	GENERATED_BODY()
+
+	//~ Begin UComputeDataInterface Interface
+	TCHAR const* GetClassName() const override { return TEXT("Dispatch"); }
+	bool IsExecutionInterface() const override { return true; }
+	void GetSupportedInputs(TArray<FShaderFunctionDefinition>& OutFunctions) const override;
+	void GetShaderParameters(TCHAR const* UID, FShaderParametersMetadataBuilder& InOutBuilder, FShaderParametersMetadataAllocations& InOutAllocations) const override;
+	const TCHAR* GetShaderVirtualPath() const override;
+	void GetShaderHash(FString& InOutKey) const override;
+	void GetHLSL(FString& OutHLSL, FString const& InDataInterfaceName) const override;
+	UComputeDataProvider* CreateDataProvider() const override;
+	//~ End UComputeDataInterface Interface
+
+private:
+	static const TCHAR* TemplateFilePath;
+};
+
+/** Compute data provider implementation for UComputeDataInterfaceDispatch. */
+UCLASS(MinimalAPI, Blueprintable, Category = ComputeFramework)
+class UDispatchDataProvider : public UComputeDataProvider
+{
+	GENERATED_BODY()
+	
+	//~ Begin UComputeDataProvider Interface
+	void Initialize(int32 InDataInterfaceIndex, UComputeDataInterface const* InDataInterface, UObject* InBinding, uint64 InInputMask, uint64 InOutputMask) override;
+	FComputeDataProviderRenderProxy* GetRenderProxy() override;
+	//~ End UComputeDataProvider Interface
+
+public:
+	UPROPERTY(EditAnywhere, BlueprintReadWrite, Category = DataInterface)
+	FIntVector ThreadCount = FIntVector(1, 1, 1);
+};
+
+/** Compute data provider proxy implementation for UComputeDataInterfaceDispatch. */
+class FDispatchDataProviderProxy : public FComputeDataProviderRenderProxy
+{
+	//~ Begin FComputeDataProviderRenderProxy Interface
+	bool IsValid(FValidationData const& InValidationData) const override;
+	int32 GetDispatchThreadCount(TArray<FIntVector>& InOutThreadCounts) const override;
+	void GatherDispatchData(FDispatchData const& InDispatchData) override;
+	//~ End FComputeDataProviderRenderProxy Interface
+
+public:
+	FDispatchDataProviderProxy(FIntVector InThreadCount);
+
+private:
+	using FParameters = FDataInterfaceDispatchParameters;
+	
+	FIntVector ThreadCount;
+};
diff --git a/Engine/Plugins/Runtime/ComputeFramework/Source/ComputeDataInterface/Private/ComputeDataInterfaceBuffer.cpp b/Engine/Plugins/Runtime/ComputeFramework/Source/ComputeDataInterface/Private/ComputeDataInterfaceBuffer.cpp
new file mode 100644
--- /dev/null
+++ b/Engine/Plugins/Runtime/ComputeFramework/Source/ComputeDataInterface/Private/ComputeDataInterfaceBuffer.cpp
@@ -0,0 +1,201 @@
+// Copyright Epic Games, Inc. All Rights Reserved.
+
+#include "ComputeFramework/ComputeDataInterfaceBuffer.h"
+
+#include "ComputeFramework/ComputeGraphComponent.h"
+#include "ComputeFramework/ShaderParamTypeDefinition.h"
+#include "RenderGraphBuilder.h"
+#include "RenderGraphUtils.h"
+#include "ShaderCompilerCore.h"
+#include "ShaderParameterMetadataBuilder.h"
+
+#include UE_INLINE_GENERATED_CPP_BY_NAME(ComputeDataInterfaceBuffer)
+
+static FShaderValueTypeHandle GetValueTypeFromEnum(EComputeDataInterfaceBufferType InEnumValueType)
+{
+	switch (InEnumValueType)
+	{
+		case EComputeDataInterfaceBufferType::Uint:		return FShaderValueType::Get(EShaderFundamentalType::Uint);
+		case EComputeDataInterfaceBufferType::Float:	return FShaderValueType::Get(EShaderFundamentalType::Float);
+		case EComputeDataInterfaceBufferType::Int2:		return FShaderValueType::Get(EShaderFundamentalType::Int, 2);
+		case EComputeDataInterfaceBufferType::Uint2:	return FShaderValueType::Get(EShaderFundamentalType::Uint, 2);
+		case EComputeDataInterfaceBufferType::Float2:	return FShaderValueType::Get(EShaderFundamentalType::Float, 2);
+		case EComputeDataInterfaceBufferType::Int3:		return FShaderValueType::Get(EShaderFundamentalType::Int, 3);
+		case EComputeDataInterfaceBufferType::Uint3:	return FShaderValueType::Get(EShaderFundamentalType::Uint, 3);
+		case EComputeDataInterfaceBufferType::Float3:	return FShaderValueType::Get(EShaderFundamentalType::Float, 3);
+		case EComputeDataInterfaceBufferType::Int4:		return FShaderValueType::Get(EShaderFundamentalType::Int, 4);
+		case EComputeDataInterfaceBufferType::Uint4:	return FShaderValueType::Get(EShaderFundamentalType::Uint, 4);
+		case EComputeDataInterfaceBufferType::Float4:	return FShaderValueType::Get(EShaderFundamentalType::Float, 4);
+		case EComputeDataInterfaceBufferType::Int:		
+		default:										return FShaderValueType::Get(EShaderFundamentalType::Int);
+	}
+}
+
+static bool SupportsAtomics(FShaderValueTypeHandle InValueType)
+{
+	return (InValueType->DimensionType == EShaderFundamentalDimensionType::Scalar && (InValueType->Type == EShaderFundamentalType::Int || InValueType->Type == EShaderFundamentalType::Uint));
+}
+
+void UComputeDataInterfaceBuffer::GetSupportedInputs(TArray<FShaderFunctionDefinition>& OutFunctions) const
+{
+	OutFunctions.AddDefaulted_GetRef()
+		.SetName(TEXT("ReadNumValues"))
+		.AddReturnType(EShaderFundamentalType::Uint);
+
+	OutFunctions.AddDefaulted_GetRef()
+		.SetName(TEXT("ReadValue"))
+		.AddReturnType(GetValueTypeFromEnum(ValueType))
+		.AddParam(EShaderFundamentalType::Uint);
+
+	OutFunctions.AddDefaulted_GetRef()
+		.SetName(TEXT("ReadValueUAV"))
+		.AddReturnType(GetValueTypeFromEnum(ValueType))
+		.AddParam(EShaderFundamentalType::Uint);
+}
+
+void UComputeDataInterfaceBuffer::GetSupportedOutputs(TArray<FShaderFunctionDefinition>& OutFunctions) const
+{
+	OutFunctions.AddDefaulted_GetRef()
+		.SetName(TEXT("WriteValue"))
+		.AddParam(EShaderFundamentalType::Uint)
+		.AddParam(GetValueTypeFromEnum(ValueType));
+
+	OutFunctions.AddDefaulted_GetRef()
+		.SetName(TEXT("WriteAtomicAdd"))
+		.AddReturnType(GetValueTypeFromEnum(ValueType))
+		.AddParam(EShaderFundamentalType::Uint)
+		.AddParam(GetValueTypeFromEnum(ValueType));
+
+	OutFunctions.AddDefaulted_GetRef()
+		.SetName(TEXT("WriteAtomicMin"))
+		.AddReturnType(GetValueTypeFromEnum(ValueType))
+		.AddParam(EShaderFundamentalType::Uint)
+		.AddParam(GetValueTypeFromEnum(ValueType));
+
+	OutFunctions.AddDefaulted_GetRef()
+		.SetName(TEXT("WriteAtomicMax"))
+		.AddReturnType(GetValueTypeFromEnum(ValueType))
+		.AddParam(EShaderFundamentalType::Uint)
+		.AddParam(GetValueTypeFromEnum(ValueType));
+}
+
+BEGIN_SHADER_PARAMETER_STRUCT(FBufferDataInterfaceParameters, )
+	SHADER_PARAMETER(uint32, BufferElementCount)
+	SHADER_PARAMETER_RDG_BUFFER_SRV(ByteAddressBuffer, BufferSRV)
+	SHADER_PARAMETER_RDG_BUFFER_UAV(RWByteAddressBuffer, BufferUAV)
+END_SHADER_PARAMETER_STRUCT()
+
+void UComputeDataInterfaceBuffer::GetShaderParameters(TCHAR const* UID, FShaderParametersMetadataBuilder& InOutBuilder, FShaderParametersMetadataAllocations& InOutAllocations) const
+{
+	InOutBuilder.AddNestedStruct<FBufferDataInterfaceParameters>(UID);
+}
+
+TCHAR const* UComputeDataInterfaceBuffer::TemplateFilePath = TEXT("/Plugin/ComputeFramework/Private/ComputeDataInterfaceBuffer.ush");
+
+TCHAR const* UComputeDataInterfaceBuffer::GetShaderVirtualPath() const
+{
+	return TemplateFilePath;
+}
+
+void UComputeDataInterfaceBuffer::GetShaderHash(FString& InOutKey) const
+{
+	GetShaderFileHash(TemplateFilePath, EShaderPlatform::SP_PCD3D_SM5).AppendString(InOutKey);
+}
+
+void UComputeDataInterfaceBuffer::GetHLSL(FString& OutHLSL, FString const& InDataInterfaceName) const
+{
+	const FString ValueTypeName = GetValueTypeFromEnum(ValueType)->ToString();
+	const int32 ValueTypeStride = GetValueTypeFromEnum(ValueType)->GetResourceElementSize();
+	const bool bSupportAtomic = SupportsAtomics(GetValueTypeFromEnum(ValueType));
+
+	TMap<FString, FStringFormatArg> TemplateArgs =
+	{
+		{ TEXT("DataInterfaceName"), InDataInterfaceName },
+		{ TEXT("ValueType"), ValueTypeName },
+		{ TEXT("ValueTypeStride"), ValueTypeStride },
+		{ TEXT("SupportAtomic"), bSupportAtomic ? 1 : 0 },
+		{ TEXT("SplitReadWrite"), bAllowReadWrite ? 0 : 1 },
+	};
+
+	FString TemplateFile;
+	LoadShaderSourceFile(TemplateFilePath, EShaderPlatform::SP_PCD3D_SM5, &TemplateFile, nullptr);
+	OutHLSL += FString::Format(*TemplateFile, TemplateArgs);
+}
+
+UComputeDataProvider* UComputeDataInterfaceBuffer::CreateDataProvider() const
+{
+	UBufferDataProvider* Provider = NewObject<UBufferDataProvider>();
+	Provider->ValueType = ValueType;
+	return Provider;
+}
+
+void UBufferDataProvider::Initialize(int32 InDataInterfaceIndex, UComputeDataInterface const* InDataInterface, UObject* InBinding, uint64 InInputMask, uint64 InOutputMask)
+{
+	if (UComputeGraphComponent* BindingComponent = Cast<UComputeGraphComponent>(InBinding))
+	{
+		BindingComponent->InitializeProvider(InDataInterfaceIndex, this);
+	}
+
+	// Override clear before use if we read but don't write the buffer.
+	if (InInputMask != 0 && InOutputMask == 0)
+	{
+		bClearBeforeUse = true;
+	}
+}
+
+FComputeDataProviderRenderProxy* UBufferDataProvider::GetRenderProxy()
+{
+	const uint32 ElementStride = GetValueTypeFromEnum(ValueType)->GetResourceElementSize();
+	return new FBufferDataProviderProxy(ElementStride, (uint32)FMath::Max(ElementCount, 0), bClearBeforeUse);
+}
+
+FBufferDataProviderProxy::FBufferDataProviderProxy(uint32 InElementStride, uint32 InElementCount, bool bInClearBeforeUse)
+	: ElementStride(InElementStride)
+	, ElementCount(InElementCount)
+	, bClearBeforeUse(bInClearBeforeUse)
+{
+}
+
+bool FBufferDataProviderProxy::IsValid(FValidationData const& InValidationData) const
+{
+	if (InValidationData.ParameterStructSize != sizeof(FParameters))
+	{
+		return false;
+	}
+	
+	if (ElementStride % 4 != 0)
+	{
+		return false;
+	}
+
+	if (ElementCount == 0)
+	{
+		return false;
+	}
+
+	return true;
+}
+
+void FBufferDataProviderProxy::AllocateResources(FRDGBuilder& GraphBuilder, FAllocationData const& InAllocationData)
+{
+	Buffer = GraphBuilder.CreateBuffer(FRDGBufferDesc::CreateByteAddressDesc(ElementStride * ElementCount), TEXT("ComputeFrameworkBuffer"), ERDGBufferFlags::None);
+	BufferSRV = GraphBuilder.CreateSRV(Buffer);
+	BufferUAV = GraphBuilder.CreateUAV(Buffer, ERDGUnorderedAccessViewFlags::SkipBarrier);
+
+	if (bClearBeforeUse)
+	{
+		AddClearUAVPass(GraphBuilder, BufferUAV, 0);
+	}
+}
+
+void FBufferDataProviderProxy::GatherDispatchData(FDispatchData const& InDispatchData)
+{
+	TStridedView<FParameters> ParameterArray = MakeStridedParameterView<FParameters>(InDispatchData);
+	for (int32 InvocationIndex = 0; InvocationIndex < ParameterArray.Num(); ++InvocationIndex)
+	{
+		FParameters& Parameters = ParameterArray[InvocationIndex];
+		Parameters.BufferElementCount = ElementCount;
+		Parameters.BufferSRV = BufferSRV;
+		Parameters.BufferUAV = BufferUAV;
+	}
+}
diff --git a/Engine/Plugins/Runtime/ComputeFramework/Source/ComputeDataInterface/Private/ComputeDataInterfaceDispatch.cpp b/Engine/Plugins/Runtime/ComputeFramework/Source/ComputeDataInterface/Private/ComputeDataInterfaceDispatch.cpp
new file mode 100644
--- /dev/null
+++ b/Engine/Plugins/Runtime/ComputeFramework/Source/ComputeDataInterface/Private/ComputeDataInterfaceDispatch.cpp
@@ -0,0 +1,102 @@
+// Copyright Epic Games, Inc. All Rights Reserved.
+
+#include "ComputeFramework/ComputeDataInterfaceDispatch.h"
+
+#include "ComputeFramework/ComputeGraphComponent.h"
+#include "ComputeFramework/ShaderParamTypeDefinition.h"
+#include "RenderGraphBuilder.h"
+#include "ShaderCompilerCore.h"
+#include "ShaderParameterMetadataBuilder.h"
+
+#include UE_INLINE_GENERATED_CPP_BY_NAME(ComputeDataInterfaceDispatch)
+
+void UComputeDataInterfaceDispatch::GetSupportedInputs(TArray<FShaderFunctionDefinition>& OutFunctions) const
+{
+	OutFunctions.AddDefaulted_GetRef()
+		.SetName(TEXT("ReadNumThreads"))
+		.AddReturnType(FShaderValueType::Get(EShaderFundamentalType::Uint, 3));
+}
+
+BEGIN_SHADER_PARAMETER_STRUCT(FDataInterfaceDispatchParameters, )
+	SHADER_PARAMETER(FUintVector3, NumThreads)
+END_SHADER_PARAMETER_STRUCT()
+
+void UComputeDataInterfaceDispatch::GetShaderParameters(TCHAR const* UID, FShaderParametersMetadataBuilder& InOutBuilder, FShaderParametersMetadataAllocations& InOutAllocations) const
+{
+	InOutBuilder.AddNestedStruct<FDataInterfaceDispatchParameters>(UID);
+}
+
+const TCHAR* UComputeDataInterfaceDispatch::TemplateFilePath = TEXT("/Plugin/ComputeFramework/Private/ComputeDataInterfaceDispatch.ush");
+
+const TCHAR* UComputeDataInterfaceDispatch::GetShaderVirtualPath() const
+{
+	return TemplateFilePath;
+}
+
+void UComputeDataInterfaceDispatch::GetShaderHash(FString& InOutKey) const
+{
+	GetShaderFileHash(TemplateFilePath, EShaderPlatform::SP_PCD3D_SM5).AppendString(InOutKey);
+}
+
+void UComputeDataInterfaceDispatch::GetHLSL(FString& OutHLSL, FString const& InDataInterfaceName) const
+{
+	TMap<FString, FStringFormatArg> TemplateArgs =
+	{
+		{ TEXT("DataInterfaceName"), InDataInterfaceName },
+	};
+
+	FString TemplateFile;
+	LoadShaderSourceFile(TemplateFilePath, EShaderPlatform::SP_PCD3D_SM5, &TemplateFile, nullptr);
+	OutHLSL += FString::Format(*TemplateFile, TemplateArgs);
+}
+
+UComputeDataProvider* UComputeDataInterfaceDispatch::CreateDataProvider() const
+{
+	return NewObject<UDispatchDataProvider>();
+}
+
+void UDispatchDataProvider::Initialize(int32 InDataInterfaceIndex, UComputeDataInterface const* InDataInterface, UObject* InBinding, uint64 InInputMask, uint64 InOutputMask)
+{
+	if (UComputeGraphComponent* BindingComponent = Cast<UComputeGraphComponent>(InBinding))
+	{
+		BindingComponent->InitializeProvider(InDataInterfaceIndex, this);
+	}
+}
+
+FComputeDataProviderRenderProxy* UDispatchDataProvider::GetRenderProxy()
+{
+	FDispatchDataProviderProxy* Proxy = new FDispatchDataProviderProxy(ThreadCount);
+	return Proxy;
+}
+
+FDispatchDataProviderProxy::FDispatchDataProviderProxy(FIntVector InThreadCount)
+	: ThreadCount(InThreadCount)
+{
+}
+
+bool FDispatchDataProviderProxy::IsValid(FValidationData const& InValidationData) const
+{
+	if (InValidationData.ParameterStructSize != sizeof(FParameters))
+	{
+		return false;
+	}
+
+	return (ThreadCount.X > 0 && ThreadCount.Y > 0 && ThreadCount.Z > 0);
+}
+
+int32 FDispatchDataProviderProxy::GetDispatchThreadCount(TArray<FIntVector>& InOutThreadCounts) const
+{
+	InOutThreadCounts.Reset();
+	InOutThreadCounts.Add(ThreadCount);
+	return 1;
+}
+
+void FDispatchDataProviderProxy::GatherDispatchData(FDispatchData const& InDispatchData)
+{
+	const TStridedView<FParameters> ParameterArray = MakeStridedParameterView<FParameters>(InDispatchData);
+	for (int32 InvocationIndex = 0; InvocationIndex < ParameterArray.Num(); ++InvocationIndex)
+	{
+		FParameters& Parameters = ParameterArray[InvocationIndex];
+		Parameters.NumThreads = FUintVector(ThreadCount.X, ThreadCount.Y, ThreadCount.Z);
+	}
+}
diff --git a/Engine/Plugins/Runtime/ComputeFramework/Source/ComputeDataInterface/Private/ComputeDataInterfaceModule.cpp b/Engine/Plugins/Runtime/ComputeFramework/Source/ComputeDataInterface/Private/ComputeDataInterfaceModule.cpp
new file mode 100644
--- /dev/null
+++ b/Engine/Plugins/Runtime/ComputeFramework/Source/ComputeDataInterface/Private/ComputeDataInterfaceModule.cpp
@@ -0,0 +1,5 @@
+// Copyright Epic Games, Inc. All Rights Reserved.
+
+#include "Modules/ModuleManager.h"
+
+IMPLEMENT_MODULE(FDefaultModuleImpl, ComputeDataInterface)
diff --git a/Engine/Plugins/Runtime/ComputeFramework/Source/ComputeFramework/Private/ComputeDataInterfaceBuffer.cpp b/Engine/Plugins/Runtime/ComputeFramework/Source/ComputeFramework/Private/ComputeDataInterfaceBuffer.cpp
deleted file mode 100644
diff --git a/Engine/Plugins/Runtime/ComputeFramework/Source/ComputeFramework/Private/ComputeDataInterfaceDispatch.cpp b/Engine/Plugins/Runtime/ComputeFramework/Source/ComputeFramework/Private/ComputeDataInterfaceDispatch.cpp
deleted file mode 100644
diff --git a/Engine/Plugins/Runtime/ComputeFramework/Source/ComputeFramework/Private/ComputeFramework/ComputeDataInterfaceBuffer.h b/Engine/Plugins/Runtime/ComputeFramework/Source/ComputeFramework/Private/ComputeFramework/ComputeDataInterfaceBuffer.h
deleted file mode 100644
diff --git a/Engine/Plugins/Runtime/ComputeFramework/Source/ComputeFramework/Private/ComputeFramework/ComputeDataInterfaceDispatch.h b/Engine/Plugins/Runtime/ComputeFramework/Source/ComputeFramework/Private/ComputeFramework/ComputeDataInterfaceDispatch.h
deleted file mode 100644
diff --git a/Engine/Plugins/Runtime/ComputeFramework/Source/ComputeFramework/Private/ComputeGraph.cpp b/Engine/Plugins/Runtime/ComputeFramework/Source/ComputeFramework/Private/ComputeGraph.cpp
--- a/Engine/Plugins/Runtime/ComputeFramework/Source/ComputeFramework/Private/ComputeGraph.cpp
+++ b/Engine/Plugins/Runtime/ComputeFramework/Source/ComputeFramework/Private/ComputeGraph.cpp
@@ -287,7 +287,7 @@ void UComputeGraph::CreateDataProviders(int32 InBindingIndex, TObjectPtr<UObject
 
 					if (DataProvider)
 					{
-						DataProvider->Initialize(DataInterface, InBindingObject, InputMask, OutputMask);
+						DataProvider->Initialize(DataInterfaceIndex, DataInterface, InBindingObject, InputMask, OutputMask);
 					}
 					PRAGMA_DISABLE_DEPRECATION_WARNINGS
 					else
diff --git a/Engine/Plugins/Runtime/ComputeFramework/Source/ComputeFramework/Private/ComputeGraphComponent.cpp b/Engine/Plugins/Runtime/ComputeFramework/Source/ComputeFramework/Private/ComputeGraphComponent.cpp
--- a/Engine/Plugins/Runtime/ComputeFramework/Source/ComputeFramework/Private/ComputeGraphComponent.cpp
+++ b/Engine/Plugins/Runtime/ComputeFramework/Source/ComputeFramework/Private/ComputeGraphComponent.cpp
@@ -3,9 +3,12 @@
 #include "ComputeFramework/ComputeGraphComponent.h"
 
 #include "ComputeFramework/ComputeFramework.h"
+#include "ComputeFramework/ComputeDataProvider.h"
+#include "ComputeFramework/ComputeGraph.h"
 #include "ComputeFramework/ComputeGraphInstance.h"
 #include "ComputeWorkerInterface.h"
 #include "GameFramework/Actor.h"
+#include "UnrealEngine.h"
 
 #include UE_INLINE_GENERATED_CPP_BY_NAME(ComputeGraphComponent)
 
@@ -21,23 +24,55 @@ UComputeGraphComponent::UComputeGraphComponent()
 void UComputeGraphComponent::OnRegister()
 {
 	Super::OnRegister();
-	CreateDataProviders(0, this);
+
+	if (ComputeGraph != nullptr)
+	{
+		const int32 NumBindingObjects = ComputeGraph->GetNumBindingObjects();
+		for (int32 BindingObjectIndex = 0; BindingObjectIndex < NumBindingObjects; ++BindingObjectIndex)
+		{
+			ComputeGraphInstance.CreateDataProviders(ComputeGraph, BindingObjectIndex, GetBindingObject(BindingObjectIndex));
+		}
+	}
 }
 
 void UComputeGraphComponent::OnUnregister()
 {
-	DestroyDataProviders();
+	ComputeGraphInstance.DestroyDataProviders();
+	
 	Super::OnUnregister();
 }
 
-void UComputeGraphComponent::CreateDataProviders(int32 InBindingIndex, UObject* InBindingObject)
+void UComputeGraphComponent::InitializeProvider_Implementation(int32 InDataInterfaceIndex, UObject* InOutDataProvider)
 {
- 	ComputeGraphInstance.CreateDataProviders(ComputeGraph, InBindingIndex, InBindingObject);
+	if (InDataInterfaceIndex < 0)
+	{
+		return;
+	}
+	
+	UComputeDataProvider* Provider = Cast<UComputeDataProvider>(InOutDataProvider);
+	if (Provider == nullptr)
+	{
+		return;
+	}
+
+	if (!DataProviderTemplates.IsValidIndex(InDataInterfaceIndex))
+	{ 
+		DataProviderTemplates.SetNum(InDataInterfaceIndex + 1);
+	}
+	if (DataProviderTemplates[InDataInterfaceIndex] == nullptr)
+	{
+		DataProviderTemplates[InDataInterfaceIndex] = DuplicateObject<UComputeDataProvider>(Provider, this);
+	}
+	else if (DataProviderTemplates[InDataInterfaceIndex]->GetClass() == Provider->GetClass())
+	{
+		UEngine::CopyPropertiesForUnrelatedObjects(DataProviderTemplates[InDataInterfaceIndex], Provider);
+	}
 }
 
-void UComputeGraphComponent::DestroyDataProviders()
+UObject* UComputeGraphComponent::GetBindingObject_Implementation(int32 InBindingIndex)
 {
-	ComputeGraphInstance.DestroyDataProviders();
+	// By default use this as the binding object.
+	return this;
 }
 
 void UComputeGraphComponent::QueueExecute()
@@ -48,13 +83,6 @@ void UComputeGraphComponent::QueueExecute()
 	}
 }
 
-void UComputeGraphComponent::TickComponent(float DeltaTime, ELevelTick TickType, FActorComponentTickFunction* ThisTickFunction)
-{
-	Super::TickComponent(DeltaTime, TickType, ThisTickFunction);
-
-	QueueExecute();
-}
-
 void UComputeGraphComponent::SendRenderDynamicData_Concurrent()
 {
 	Super::SendRenderDynamicData_Concurrent();
diff --git a/Engine/Plugins/Runtime/ComputeFramework/Source/ComputeFramework/Public/ComputeFramework/ComputeDataProvider.h b/Engine/Plugins/Runtime/ComputeFramework/Source/ComputeFramework/Public/ComputeFramework/ComputeDataProvider.h
--- a/Engine/Plugins/Runtime/ComputeFramework/Source/ComputeFramework/Public/ComputeFramework/ComputeDataProvider.h
+++ b/Engine/Plugins/Runtime/ComputeFramework/Source/ComputeFramework/Public/ComputeFramework/ComputeDataProvider.h
@@ -26,7 +26,11 @@ class UComputeDataProvider : public UObject
 
 public:
 	/** Set up the data provider from the given data interface and binding. */
-	virtual void Initialize(const UComputeDataInterface* InDataInterface, UObject* InBinding, uint64 InInputMask, uint64 InOutputMask) {}
+	virtual void Initialize(UComputeDataInterface const* InDataInterface, UObject* InBinding, uint64 InInputMask, uint64 InOutputMask) {}
+	virtual void Initialize(int32 InDataInterfaceIndex, UComputeDataInterface const* InDataInterface, UObject* InBinding, uint64 InInputMask, uint64 InOutputMask)
+	{
+		Initialize(InDataInterface, InBinding, InInputMask, InOutputMask);
+	}
 
 	/**
 	 * Get an associated render thread proxy object.
diff --git a/Engine/Plugins/Runtime/ComputeFramework/Source/ComputeFramework/Public/ComputeFramework/ComputeGraph.h b/Engine/Plugins/Runtime/ComputeFramework/Source/ComputeFramework/Public/ComputeFramework/ComputeGraph.h
--- a/Engine/Plugins/Runtime/ComputeFramework/Source/ComputeFramework/Public/ComputeFramework/ComputeGraph.h
+++ b/Engine/Plugins/Runtime/ComputeFramework/Source/ComputeFramework/Public/ComputeFramework/ComputeGraph.h
@@ -107,6 +107,9 @@ public:
 	 */
 	COMPUTEFRAMEWORK_API bool ValidateGraph(FString* OutErrors = nullptr);
 
+	/** Get the number of binding objects used by the graph. */
+	int32 GetNumBindingObjects() const { return Bindings.Num(); }
+
 	/** 
 	 * Create UComputeDataProvider objects to match the current UComputeDataInterface objects and initialize them. 
 	 * We attempt to setup bindings from the InBindingObjects.
diff --git a/Engine/Plugins/Runtime/ComputeFramework/Source/ComputeFramework/Public/ComputeFramework/ComputeGraphComponent.h b/Engine/Plugins/Runtime/ComputeFramework/Source/ComputeFramework/Public/ComputeFramework/ComputeGraphComponent.h
--- a/Engine/Plugins/Runtime/ComputeFramework/Source/ComputeFramework/Public/ComputeFramework/ComputeGraphComponent.h
+++ b/Engine/Plugins/Runtime/ComputeFramework/Source/ComputeFramework/Public/ComputeFramework/ComputeGraphComponent.h
@@ -8,13 +8,11 @@
 
 #define UE_API COMPUTEFRAMEWORK_API
 
-class UComputeGraph;
-
 /** 
  * Component which holds a context for a UComputeGraph.
- * This object binds the graph to its data providers, and queues the execution. 
+ * This object creates the graph data providers, and queues the execution. 
  */
-UCLASS(MinimalAPI, Blueprintable, meta = (BlueprintSpawnableComponent))
+UCLASS(MinimalAPI, Blueprintable, Category = ComputeFramework, meta = (BlueprintSpawnableComponent))
 class UComputeGraphComponent : public UActorComponent
 {
 	GENERATED_BODY()
@@ -22,35 +20,38 @@ class UComputeGraphComponent : public UActorComponent
 public:
 	UE_API UComputeGraphComponent();
 
+	/** Initialize a Data Provider object for the ComputeGraph. */
+	UFUNCTION(BlueprintNativeEvent, Category = "Compute")
+	UE_API void InitializeProvider(int32 InDataInterfaceIndex, UObject* InOutDataProvider);
+
+	/** Queue the graph for execution at the next render update. */
+	UFUNCTION(BlueprintCallable, Category = "Compute")
+	UE_API void QueueExecute();
+
+protected:
 	/** The Compute Graph asset. */
 	UPROPERTY(EditAnywhere, BlueprintReadWrite, Category = "Compute")
 	TObjectPtr<UComputeGraph> ComputeGraph = nullptr;
 
-	/** Create all the Data Provider objects for a given binding object of the ComputeGraph. */
-	UFUNCTION(BlueprintCallable, Category = "Compute")
-	UE_API void CreateDataProviders(int32 InBindingIndex, UObject* InBindingObject);
+	/** Get the binding object for the ComputeGraph. */
+	UFUNCTION(BlueprintNativeEvent, Category = "Compute")
+	UObject* GetBindingObject(int32 InBindingIndex);
 
-	/** Destroy all associated DataProvider objects. */
-	UFUNCTION(BlueprintCallable, Category = "Compute")
-	UE_API void DestroyDataProviders();
+	/** Serialized data provider properties. */
+	UPROPERTY(EditAnywhere, Category = "Compute", Instanced, NoClear, meta = (ShowOnlyInnerProperties))
+	TArray<TObjectPtr<UComputeDataProvider>> DataProviderTemplates;
 
-	/** Queue the graph for execution at the next render update. */
-	UFUNCTION(BlueprintCallable, Category = "Compute")
-	UE_API void QueueExecute();
+	/** The Instance with associated per instance settings. */
+	UPROPERTY()
+	FComputeGraphInstance ComputeGraphInstance;
 
-protected:
 	//~ Begin UActorComponent Interface
 	void OnRegister() override;
 	void OnUnregister() override;
-	void TickComponent(float DeltaTime, enum ELevelTick TickType, FActorComponentTickFunction* ThisTickFunction) override;
 	void SendRenderDynamicData_Concurrent() override;
 	void DestroyRenderState_Concurrent() override;
 	bool ShouldCreateRenderState() const override { return true; }
 	//~ End UActorComponent Interface
-
-private:
-	UPROPERTY()
-	FComputeGraphInstance ComputeGraphInstance;
 };
 
 #undef UE_API
