ComputeFramework: Tidy up move some files from header to cpp directory.
#rnx
#Synced-CL 52056483

diff --git a/Engine/Plugins/Runtime/ComputeFramework/Source/ComputeFramework/Private/ComputeDataProvider.cpp b/Engine/Plugins/Runtime/ComputeFramework/Source/ComputeFramework/Private/ComputeDataProvider.cpp
new file mode 100644
--- /dev/null
+++ b/Engine/Plugins/Runtime/ComputeFramework/Source/ComputeFramework/Private/ComputeDataProvider.cpp
@@ -0,0 +1,19 @@
+// Copyright Epic Games, Inc. All Rights Reserved.
+
+#include "ComputeFramework/ComputeDataProvider.h"
+
+#include "ComputeWorkerInterface.h"
+
+#include UE_INLINE_GENERATED_CPP_BY_NAME(ComputeDataProvider)
+
+void FComputeDataProviderRenderProxy::PreSubmit(FRDGBuilder& InGraphBuilder) const
+{
+	FComputeContext Context = { .GraphBuilder = InGraphBuilder };
+	PreSubmit(Context);
+}
+
+void FComputeDataProviderRenderProxy::PostSubmit(FRDGBuilder& InGraphBuilder) const
+{
+	FComputeContext Context = { .GraphBuilder = InGraphBuilder };
+	PostSubmit(Context);
+}
diff --git a/Engine/Plugins/Runtime/ComputeFramework/Source/ComputeFramework/Private/ComputeFramework/ComputeDataProvider.cpp b/Engine/Plugins/Runtime/ComputeFramework/Source/ComputeFramework/Private/ComputeFramework/ComputeDataProvider.cpp
deleted file mode 100644
diff --git a/Engine/Plugins/Runtime/ComputeFramework/Source/ComputeFramework/Private/ComputeFramework/ComputeViewExtension.cpp b/Engine/Plugins/Runtime/ComputeFramework/Source/ComputeFramework/Private/ComputeFramework/ComputeViewExtension.cpp
deleted file mode 100644
diff --git a/Engine/Plugins/Runtime/ComputeFramework/Source/ComputeFramework/Private/ComputeViewExtension.cpp b/Engine/Plugins/Runtime/ComputeFramework/Source/ComputeFramework/Private/ComputeViewExtension.cpp
new file mode 100644
--- /dev/null
+++ b/Engine/Plugins/Runtime/ComputeFramework/Source/ComputeFramework/Private/ComputeViewExtension.cpp
@@ -0,0 +1,68 @@
+// Copyright Epic Games, Inc. All Rights Reserved.
+
+#include "ComputeFramework/ComputeViewExtension.h"
+
+#include "ComputeWorkerInterface.h"
+
+FComputeViewExtension::FComputeViewExtension(const FAutoRegister& AutoReg, UWorld* InWorld)
+	: FWorldSceneViewExtension(AutoReg, InWorld)
+{
+}
+
+void FComputeViewExtension::PostTLASBuild_RenderThread(FRDGBuilder& GraphBuilder, FSceneView& InView)
+{
+	const TWeakObjectPtr<UWorld> WorldPtr = GetWorld();
+	if (!WorldPtr.IsValid())
+	{
+		return;
+	}
+
+	const FSceneInterface* Scene = WorldPtr->Scene;
+	if (!Scene)
+	{
+		return;
+	}
+
+	const ERHIFeatureLevel::Type FeatureLevel = Scene->GetFeatureLevel();
+
+	ComputeWorkers.Empty(ComputeWorkers.Num());
+	Scene->GetComputeTaskWorkers(ComputeWorkers);
+
+	for (IComputeTaskWorker* ComputeTaskWorker : ComputeWorkers)
+	{
+		if (ComputeTaskWorker->HasWork(ComputeTaskExecutionGroup::PostTLASBuild))
+		{
+			FComputeContext Context
+			{
+				.GraphBuilder = GraphBuilder,
+				.ExecutionGroupName = ComputeTaskExecutionGroup::PostTLASBuild,
+				.FeatureLevel = FeatureLevel,
+				.Scene = Scene,
+				.View = &InView,
+			};
+
+			ComputeTaskWorker->SubmitWork(Context);
+		}
+	}
+}
+
+ESceneViewExtensionFlags FComputeViewExtension::GetFlags() const
+{
+	const TWeakObjectPtr<UWorld> WorldPtr = GetWorld();
+
+	if (const FSceneInterface* Scene = WorldPtr.IsValid() ? WorldPtr->Scene : nullptr)
+	{
+		ComputeWorkers.Empty(ComputeWorkers.Num());
+		Scene->GetComputeTaskWorkers(ComputeWorkers);
+
+		for (IComputeTaskWorker* ComputeTaskWorker : ComputeWorkers)
+		{
+			if (ComputeTaskWorker->HasWork(ComputeTaskExecutionGroup::PostTLASBuild))
+			{
+				return ESceneViewExtensionFlags::SubscribesToPostTLASBuild;
+			}
+		}
+	}
+
+	return ESceneViewExtensionFlags::None;
+}
