ComputeFramework: Add builder utility class for building a UComputeGraph programatically.
#rnx
#Synced-CL 52093060

diff --git a/Engine/Plugins/Runtime/ComputeFramework/Source/ComputeFramework/Private/ComputeGraphBuilder.cpp b/Engine/Plugins/Runtime/ComputeFramework/Source/ComputeFramework/Private/ComputeGraphBuilder.cpp
new file mode 100644
--- /dev/null
+++ b/Engine/Plugins/Runtime/ComputeFramework/Source/ComputeFramework/Private/ComputeGraphBuilder.cpp
@@ -0,0 +1,159 @@
+// Copyright Epic Games, Inc. All Rights Reserved.
+
+#include "ComputeFramework/ComputeGraphBuilder.h"
+
+#include "Algo/Count.h"
+#include "ComputeFramework/ComputeDataInterface.h"
+#include "ComputeFramework/ComputeGraph.h"
+#include "ComputeFramework/ComputeKernel.h"
+#include "ComputeFramework/ComputeKernelSource.h"
+#include "ComputeFramework/ShaderParamTypeDefinition.h"
+
+#if WITH_EDITOR
+
+FKernelHandle FComputeGraphBuilder::AddKernel(UComputeKernelSource* Source, FName Name)
+{
+	if (!Source)
+	{
+		return FKernelHandle{};
+	}
+
+	FKernelEntry& Entry = KernelEntries.AddDefaulted_GetRef();
+	Entry.Source = Source;
+	Entry.Name = Name;
+	return FKernelHandle{ KernelEntries.Num() - 1 };
+}
+
+FInterfaceHandle FComputeGraphBuilder::AddDataInterface(UComputeDataInterface* Interface, UClass* BindingClass, FName Name)
+{
+	if (!Interface || !BindingClass)
+	{
+		return FInterfaceHandle{};
+	}
+
+	FInterfaceEntry& Entry = InterfaceEntries.AddDefaulted_GetRef();
+	Entry.Interface = Interface;
+	Entry.BindingClass = BindingClass;
+	Entry.Name = Name;
+	return FInterfaceHandle{ InterfaceEntries.Num() - 1 };
+}
+
+void FComputeGraphBuilder::ConnectInput(FKernelHandle Kernel, FStringView KernelFn, FInterfaceHandle Interface, FStringView InterfaceFn)
+{
+	if (!KernelEntries.IsValidIndex(Kernel.Index) || !InterfaceEntries.IsValidIndex(Interface.Index) || KernelFn.IsEmpty() || InterfaceFn.IsEmpty())
+	{
+		return;
+	}
+
+	FPinEntry& Pin = KernelEntries[Kernel.Index].Inputs.AddDefaulted_GetRef();
+	Pin.KernelFunctionName = KernelFn;
+	Pin.InterfaceIndex = Interface.Index;
+	Pin.InterfaceFunctionName = InterfaceFn;
+}
+
+void FComputeGraphBuilder::ConnectOutput(FKernelHandle Kernel, FStringView KernelFn, FInterfaceHandle Interface, FStringView InterfaceFn)
+{
+	if (!KernelEntries.IsValidIndex(Kernel.Index) || !InterfaceEntries.IsValidIndex(Interface.Index) || KernelFn.IsEmpty() || InterfaceFn.IsEmpty())
+	{
+		return;
+	}
+
+	FPinEntry& Pin = KernelEntries[Kernel.Index].Outputs.AddDefaulted_GetRef();
+	Pin.KernelFunctionName = KernelFn;
+	Pin.InterfaceIndex = Interface.Index;
+	Pin.InterfaceFunctionName = InterfaceFn;
+}
+
+void FComputeGraphBuilder::Build(UComputeGraph& OutGraph, UObject& Outer) const
+{
+	OutGraph.KernelInvocations.Reset();
+	OutGraph.DataInterfaces.Reset();
+	OutGraph.GraphEdges.Reset();
+	OutGraph.Bindings.Reset();
+	OutGraph.DataInterfaceToBinding.Reset();
+
+	// Register unique binding classes, preserving insertion order.
+	TMap<UClass*, int32> BindingClassToIndex;
+	for (FInterfaceEntry const& Entry : InterfaceEntries)
+	{
+		if (!BindingClassToIndex.Contains(Entry.BindingClass))
+		{
+			BindingClassToIndex.Add(Entry.BindingClass, OutGraph.Bindings.Add(Entry.BindingClass));
+		}
+	}
+
+	// Add data interfaces.
+	for (FInterfaceEntry const& Entry : InterfaceEntries)
+	{
+		OutGraph.DataInterfaces.Add(Entry.Interface);
+		OutGraph.DataInterfaceToBinding.Add(BindingClassToIndex[Entry.BindingClass]);
+	}
+
+	// Add kernels and wire their pins.
+	for (FKernelEntry const& KernelEntry : KernelEntries)
+	{
+		UComputeKernel* Kernel = NewObject<UComputeKernel>(&Outer);
+		Kernel->KernelSource = KernelEntry.Source;
+		const int32 KernelIndex = OutGraph.KernelInvocations.Add(Kernel);
+
+		auto AddConnection = [&](FPinEntry const& Pin, bool const bInput)
+		{
+			if (!OutGraph.DataInterfaces.IsValidIndex(Pin.InterfaceIndex))
+			{
+				return;
+			}
+
+			TArray<FShaderFunctionDefinition> Functions;
+			if (bInput)
+			{
+				OutGraph.DataInterfaces[Pin.InterfaceIndex]->GetSupportedInputs(Functions);
+			}
+			else
+			{
+				OutGraph.DataInterfaces[Pin.InterfaceIndex]->GetSupportedOutputs(Functions);
+			}
+
+			const int32 FunctionIndex = Functions.IndexOfByPredicate([&Pin](FShaderFunctionDefinition const& Def) { return Def.Name == Pin.InterfaceFunctionName; });
+			if (FunctionIndex == INDEX_NONE)
+			{
+				return;
+			}
+
+			if (bInput)
+			{
+				Kernel->KernelSource->ExternalInputs.Add_GetRef(Functions[FunctionIndex]).SetName(Pin.KernelFunctionName);
+			}
+			else
+			{
+				Kernel->KernelSource->ExternalOutputs.Add_GetRef(Functions[FunctionIndex]).SetName(Pin.KernelFunctionName);
+			}
+
+			const int32 KernelBindingIndex = Algo::CountIf(OutGraph.GraphEdges,	[KernelIndex, bInput](FComputeGraphEdge const& Edge)
+			{
+				return Edge.KernelIndex == KernelIndex && Edge.bKernelInput == bInput;
+			});
+
+			FComputeGraphEdge& Edge = OutGraph.GraphEdges.AddDefaulted_GetRef();
+			Edge.KernelIndex = KernelIndex;
+			Edge.KernelBindingIndex = KernelBindingIndex;
+			Edge.DataInterfaceIndex = Pin.InterfaceIndex;
+			Edge.DataInterfaceBindingIndex = FunctionIndex;
+			Edge.bKernelInput = bInput;
+		};
+
+		Kernel->KernelSource->ExternalInputs.Reset();
+		for (FPinEntry const& Pin : KernelEntry.Inputs)
+		{
+			AddConnection(Pin, true);
+		}
+
+		Kernel->KernelSource->ExternalOutputs.Reset();
+		for (FPinEntry const& Pin : KernelEntry.Outputs)
+		{ 
+			AddConnection(Pin, false);
+		}
+	}
+}
+
+
+#endif // WITH_EDITOR
diff --git a/Engine/Plugins/Runtime/ComputeFramework/Source/ComputeFramework/Public/ComputeFramework/ComputeGraph.h b/Engine/Plugins/Runtime/ComputeFramework/Source/ComputeFramework/Public/ComputeFramework/ComputeGraph.h
--- a/Engine/Plugins/Runtime/ComputeFramework/Source/ComputeFramework/Public/ComputeFramework/ComputeGraph.h
+++ b/Engine/Plugins/Runtime/ComputeFramework/Source/ComputeFramework/Public/ComputeFramework/ComputeGraph.h
@@ -11,6 +11,7 @@
 
 enum EShaderPlatform : uint16;
 class FArchive;
+class FComputeGraphBuilder;
 struct FComputeKernelCompileResults;
 struct FComputeKernelDefinitionSet;
 struct FComputeKernelPermutationVector;
@@ -72,6 +73,8 @@ class UComputeGraph : public UObject
 {
 	GENERATED_BODY()
 
+	friend class FComputeGraphBuilder;
+
 protected:
 	/** Kernels in the graph. */
 	UPROPERTY()
diff --git a/Engine/Plugins/Runtime/ComputeFramework/Source/ComputeFramework/Public/ComputeFramework/ComputeGraphBuilder.h b/Engine/Plugins/Runtime/ComputeFramework/Source/ComputeFramework/Public/ComputeFramework/ComputeGraphBuilder.h
new file mode 100644
--- /dev/null
+++ b/Engine/Plugins/Runtime/ComputeFramework/Source/ComputeFramework/Public/ComputeFramework/ComputeGraphBuilder.h
@@ -0,0 +1,89 @@
+// Copyright Epic Games, Inc. All Rights Reserved.
+
+#pragma once
+
+#include "CoreMinimal.h"
+
+class UComputeGraph;
+class UComputeKernelSource;
+class UComputeDataInterface;
+
+#define UE_API COMPUTEFRAMEWORK_API
+
+#if WITH_EDITOR
+
+/** Opaque handle to a kernel added to FComputeGraphBuilder. */
+struct FKernelHandle
+{
+	int32 Index = INDEX_NONE;
+};
+
+/** Opaque handle to a data interface added to FComputeGraphBuilder. */
+struct FInterfaceHandle
+{
+	int32 Index = INDEX_NONE;
+};
+
+/**
+ * Programmatic builder for UComputeGraph.
+ * Handles the direct array manipulation and name-matching that are required when directly filling the UComputeGraph properties. 
+ * Expected to be used as a temporary object to do the build. Passed in UObject pointers aren't GC tracked.
+ */
+class FComputeGraphBuilder
+{
+public:
+	FComputeGraphBuilder() = default;
+	FComputeGraphBuilder(const FComputeGraphBuilder&) = delete;
+	FComputeGraphBuilder& operator=(const FComputeGraphBuilder&) = delete;
+
+	/** 
+	 * Add a kernel. Source must outlive the builder (typically owned by a UObject outer). 
+	 * Source will be modified by Build() which resets it's inputs and outputs to the new ones applied by the builder.
+	 */
+	UE_API FKernelHandle AddKernel(UComputeKernelSource* Source, FName Name = NAME_None);
+
+	/** Add a data interface. BindingClass is the UObject subclass that will be bound at runtime. */
+	UE_API FInterfaceHandle AddDataInterface(UComputeDataInterface* Interface, UClass* BindingClass, FName Name = NAME_None);
+
+	/** Connect a kernel HLSL input parameter to a data interface read function. */
+	UE_API void ConnectInput(FKernelHandle Kernel, FStringView KernelFn, FInterfaceHandle Interface, FStringView InterfaceFn);
+
+	/** Connect a kernel HLSL output parameter to a data interface write function. */
+	UE_API void ConnectOutput(FKernelHandle Kernel, FStringView KernelFn, FInterfaceHandle Interface, FStringView InterfaceFn);
+
+	/**
+	 * Apply the built topology to OutGraph, resetting any existing topology first.
+	 * Outer is used as the UObject outer for the UComputeKernel wrapper objects created internally.
+	 */
+	UE_API void Build(UComputeGraph& OutGraph, UObject& Outer) const;
+
+private:
+	struct FPinEntry
+	{
+		FString KernelFunctionName;
+		int32 InterfaceIndex = INDEX_NONE;
+		FString InterfaceFunctionName;
+	};
+
+	struct FKernelEntry
+	{
+		UComputeKernelSource* Source;
+		FName Name;
+		TArray<FPinEntry> Inputs;
+		TArray<FPinEntry> Outputs;
+	};
+
+	struct FInterfaceEntry
+	{
+		UComputeDataInterface* Interface;
+		UClass* BindingClass;
+		FName Name;
+	};
+
+	TArray<FKernelEntry> KernelEntries;
+	TArray<FInterfaceEntry> InterfaceEntries;
+};
+
+#endif // WITH_EDITOR
+
+#undef UE_API
