ComputeFramework: Improve graph validation.
Add validation for mismatching function signatures, multiple inputs to one pin and graph being cyclic.
Also output descriptive text for any validation errors.
#Synced-CL 52204248

diff --git a/Engine/Plugins/Runtime/ComputeFramework/Source/ComputeFramework/Private/ComputeGraph.cpp b/Engine/Plugins/Runtime/ComputeFramework/Source/ComputeFramework/Private/ComputeGraph.cpp
--- a/Engine/Plugins/Runtime/ComputeFramework/Source/ComputeFramework/Private/ComputeGraph.cpp
+++ b/Engine/Plugins/Runtime/ComputeFramework/Source/ComputeFramework/Private/ComputeGraph.cpp
@@ -13,6 +13,7 @@
 #include "ComputeFramework/ComputeKernelShared.h"
 #include "ComputeFramework/ComputeKernelSource.h"
 #include "ComputeFramework/ComputeSource.h"
+#include "ComputeFramework/ShaderParamTypeDefinition.h"
 #include "ComputeFramework/ShaderParameterMetadataAllocation.h"
 #include "DataDrivenShaderPlatformInfo.h"
 #include "GameFramework/Actor.h"
@@ -209,40 +210,205 @@ void UComputeGraph::OnCookEvent(UE::Cook::ECookEvent CookEvent, UE::Cook::FCookE
 
 bool UComputeGraph::ValidateGraph(FString* OutErrors)
 {
-	// todo[CF]:
-	// Check same number of kernel in/outs as edges.
-	// Check each edge connects matching function types.
-	// Check graph is DAG
+	bool bValid = true;
 
-	// Validate that we have one execution provider per kernel.
+	// Validate that we have exactly one execution provider per kernel.
 	TArray<bool, TInlineAllocator<64>> KernelHasExecution;
-	KernelHasExecution.SetNumUninitialized(KernelInvocations.Num());
+	KernelHasExecution.SetNumZeroed(KernelInvocations.Num());
+	for (FComputeGraphEdge const& Edge : GraphEdges)
+	{
+		if (!KernelInvocations.IsValidIndex(Edge.KernelIndex))
+		{
+			continue;
+		}
+		if (DataInterfaces.IsValidIndex(Edge.DataInterfaceIndex) && DataInterfaces[Edge.DataInterfaceIndex] != nullptr
+			&& DataInterfaces[Edge.DataInterfaceIndex]->IsExecutionInterface())
+		{
+			if (KernelHasExecution[Edge.KernelIndex])
+			{
+				if (OutErrors && KernelInvocations.IsValidIndex(Edge.KernelIndex)
+					&& KernelInvocations[Edge.KernelIndex] && KernelInvocations[Edge.KernelIndex]->KernelSource)
+				{
+					OutErrors->Appendf(TEXT("Kernel '%s': multiple execution interfaces connected.\n"),
+						*KernelInvocations[Edge.KernelIndex]->KernelSource->EntryPoint);
+				}
+				return false;
+			}
+			KernelHasExecution[Edge.KernelIndex] = true;
+		}
+	}
 	for (int32 KernelIndex = 0; KernelIndex < KernelInvocations.Num(); ++KernelIndex)
 	{
-		KernelHasExecution[KernelIndex] = false;
+		if (KernelInvocations[KernelIndex] != nullptr && !KernelHasExecution[KernelIndex])
+		{
+			if (OutErrors && KernelInvocations[KernelIndex]->KernelSource)
+			{
+				OutErrors->Appendf(TEXT("Kernel '%s': no execution interface (e.g. Dispatch) connected. Bind a pin to a Dispatch interface.\n"),
+					*KernelInvocations[KernelIndex]->KernelSource->EntryPoint);
+			}
+			bValid = false;
+		}
 	}
-	for (int32 GraphEdgeIndex = 0; GraphEdgeIndex < GraphEdges.Num(); ++GraphEdgeIndex)
+
+	// Check that no two edges occupy the same kernel binding slot (duplicate bindings).
 	{
-		const int32 DataInterfaceIndex = GraphEdges[GraphEdgeIndex].DataInterfaceIndex;
-		if (DataInterfaces[DataInterfaceIndex]->IsExecutionInterface())
+		TSet<uint64> InputsSeen, OutputsSeen;
+		for (FComputeGraphEdge const& Edge : GraphEdges)
 		{
-			const int32 KernelIndex = GraphEdges[GraphEdgeIndex].KernelIndex;
-			if (KernelHasExecution[KernelIndex])
+			const uint64 Key = ((uint64)(uint32)Edge.KernelIndex << 32) | (uint32)Edge.KernelBindingIndex;
+			TSet<uint64>& Seen = Edge.bKernelInput ? InputsSeen : OutputsSeen;
+			bool bAlreadyInSet = false;
+			Seen.FindOrAdd(Key, &bAlreadyInSet);
+			if (bAlreadyInSet)
 			{
-				return false;
+				if (OutErrors && KernelInvocations.IsValidIndex(Edge.KernelIndex)
+					&& KernelInvocations[Edge.KernelIndex] && KernelInvocations[Edge.KernelIndex]->KernelSource)
+				{
+					OutErrors->Appendf(TEXT("Kernel '%s': duplicate %s binding at index %d.\n"),
+						*KernelInvocations[Edge.KernelIndex]->KernelSource->EntryPoint,
+						Edge.bKernelInput ? TEXT("input") : TEXT("output"),
+						Edge.KernelBindingIndex);
+				}
+				bValid = false;
 			}
-			KernelHasExecution[KernelIndex] = true;
 		}
 	}
-	for (int32 KernelIndex = 0; KernelIndex < KernelInvocations.Num(); ++KernelIndex)
+
+	// Check that each edge's DataInterfaceBindingIndex is valid and (editor-only) that function types match.
+	for (FComputeGraphEdge const& Edge : GraphEdges)
 	{
-		if (KernelInvocations[KernelIndex] != nullptr && KernelHasExecution[KernelIndex] == false)
+		if (!DataInterfaces.IsValidIndex(Edge.DataInterfaceIndex) || !DataInterfaces[Edge.DataInterfaceIndex])
 		{
-			return false;
+			continue;
+		}
+
+		TArray<FShaderFunctionDefinition> Functions;
+		if (Edge.bKernelInput)
+		{
+			DataInterfaces[Edge.DataInterfaceIndex]->GetSupportedInputs(Functions);
+		}
+		else
+		{
+			DataInterfaces[Edge.DataInterfaceIndex]->GetSupportedOutputs(Functions);
+		}
+
+		if (!Functions.IsValidIndex(Edge.DataInterfaceBindingIndex))
+		{
+			if (OutErrors)
+			{
+				OutErrors->Appendf(TEXT("Edge references out-of-range data interface binding index %d (max %d).\n"),
+					Edge.DataInterfaceBindingIndex, Functions.Num() - 1);
+			}
+			bValid = false;
+			continue;
+		}
+
+#if WITH_EDITORONLY_DATA
+		if (KernelInvocations.IsValidIndex(Edge.KernelIndex) && KernelInvocations[Edge.KernelIndex]	&& KernelInvocations[Edge.KernelIndex]->KernelSource)
+		{
+			TArray<FShaderFunctionDefinition> const& KernelFns = Edge.bKernelInput
+				? KernelInvocations[Edge.KernelIndex]->KernelSource->ExternalInputs
+				: KernelInvocations[Edge.KernelIndex]->KernelSource->ExternalOutputs;
+
+			if (KernelFns.IsValidIndex(Edge.KernelBindingIndex))
+			{
+				FShaderFunctionDefinition const& DIFn = Functions[Edge.DataInterfaceBindingIndex];
+				FShaderFunctionDefinition const& KernelFn = KernelFns[Edge.KernelBindingIndex];
+
+				bool bParamTypesMatch = (DIFn.ParamTypes.Num() == KernelFn.ParamTypes.Num());
+				for (int32 ParamIdx = 0; bParamTypesMatch && ParamIdx < DIFn.ParamTypes.Num(); ++ParamIdx)
+				{
+					bParamTypesMatch = DIFn.ParamTypes[ParamIdx].TypeDeclaration == KernelFn.ParamTypes[ParamIdx].TypeDeclaration;
+				}
+
+				if (!bParamTypesMatch)
+				{
+					if (OutErrors)
+					{
+						OutErrors->Appendf(TEXT("Kernel '%s': %s pin %d type mismatch with data interface function '%s'.\n"),
+							*KernelInvocations[Edge.KernelIndex]->KernelSource->EntryPoint,
+							Edge.bKernelInput ? TEXT("input") : TEXT("output"),
+							Edge.KernelBindingIndex,
+							*DIFn.Name);
+					}
+					bValid = false;
+				}
+			}
 		}
+#endif
 	}
 
-	return true;
+	// Check that the kernel dependency graph is a DAG (no cycles).
+	{
+		const int32 NumKernels = KernelInvocations.Num();
+		TArray<TArray<int32>> Adj;
+		Adj.SetNum(NumKernels);
+
+		TMap<int32, TArray<int32>> DiWriters; // DataInterfaceIndex to kernels that output to it.
+		for (FComputeGraphEdge const& Edge : GraphEdges)
+		{
+			if (!Edge.bKernelInput && KernelInvocations.IsValidIndex(Edge.KernelIndex))
+			{
+				DiWriters.FindOrAdd(Edge.DataInterfaceIndex).AddUnique(Edge.KernelIndex);
+			}
+		}
+		for (FComputeGraphEdge const& Edge : GraphEdges)
+		{
+			if (Edge.bKernelInput && KernelInvocations.IsValidIndex(Edge.KernelIndex))
+			{
+				if (TArray<int32>* Writers = DiWriters.Find(Edge.DataInterfaceIndex))
+				{
+					for (int32 WriterKernel : *Writers)
+					{
+						if (WriterKernel != Edge.KernelIndex)
+						{
+							Adj[WriterKernel].AddUnique(Edge.KernelIndex);
+						}
+					}
+				}
+			}
+		}
+
+		// Three-color DFS: 0=unvisited, 1=in-stack, 2=done.
+		TArray<uint8> Color;
+		Color.SetNumZeroed(NumKernels);
+
+		TFunction<void(int32)> Dfs = [&](int32 Node)
+		{
+			Color[Node] = 1;
+			for (int32 Next : Adj[Node])
+			{
+				if (Color[Next] == 1)
+				{
+					if (OutErrors && KernelInvocations.IsValidIndex(Node) && KernelInvocations[Node]
+						&& KernelInvocations[Node]->KernelSource
+						&& KernelInvocations.IsValidIndex(Next) && KernelInvocations[Next]
+						&& KernelInvocations[Next]->KernelSource)
+					{
+						OutErrors->Appendf(TEXT("Graph cycle detected: kernel '%s' to '%s'.\n"),
+							*KernelInvocations[Node]->KernelSource->EntryPoint,
+							*KernelInvocations[Next]->KernelSource->EntryPoint);
+					}
+					bValid = false;
+				}
+				else if (Color[Next] == 0)
+				{
+					Dfs(Next);
+				}
+			}
+			Color[Node] = 2;
+		};
+
+		for (int32 KernelIndex = 0; KernelIndex < NumKernels; ++KernelIndex)
+		{
+			if (Color[KernelIndex] == 0)
+			{
+				Dfs(KernelIndex);
+			}
+		}
+	}
+
+	return bValid;
 }
 
 bool UComputeGraph::ValidateProviders(TArray< TObjectPtr<UComputeDataProvider> > const& DataProviders) const
