ComputeFramework: Add simple editable compute graph and toolkit to edit it.
This is intended to be self supporting and lighter weight than other specialized ComputeGraph editors such as the DeformerGraph system.
#rnx
#Synced-CL 52204248

diff --git a/Engine/Plugins/Runtime/ComputeFramework/ComputeFramework.uplugin b/Engine/Plugins/Runtime/ComputeFramework/ComputeFramework.uplugin
--- a/Engine/Plugins/Runtime/ComputeFramework/ComputeFramework.uplugin
+++ b/Engine/Plugins/Runtime/ComputeFramework/ComputeFramework.uplugin
@@ -29,6 +29,16 @@
 			"Name": "ComputeDataInterface",
 			"Type": "Runtime",
 			"LoadingPhase": "Default"
+		},
+		{
+			"Name": "EditableComputeGraph",
+			"Type": "Runtime",
+			"LoadingPhase": "Default"
+		},
+		{
+			"Name": "EditableComputeGraphEditor",
+			"Type": "Editor",
+			"LoadingPhase": "Default"
 		}
 	]
 }
diff --git a/Engine/Plugins/Runtime/ComputeFramework/Source/EditableComputeGraph/EditableComputeGraph.Build.cs b/Engine/Plugins/Runtime/ComputeFramework/Source/EditableComputeGraph/EditableComputeGraph.Build.cs
new file mode 100644
--- /dev/null
+++ b/Engine/Plugins/Runtime/ComputeFramework/Source/EditableComputeGraph/EditableComputeGraph.Build.cs
@@ -0,0 +1,29 @@
+// Copyright Epic Games, Inc. All Rights Reserved.
+
+namespace UnrealBuildTool.Rules
+{
+	public class EditableComputeGraph : ModuleRules
+	{
+		public EditableComputeGraph(ReadOnlyTargetRules Target) : base(Target)
+		{
+			PublicDependencyModuleNames.AddRange(
+				new string[]
+				{
+					"Core",
+					"CoreUObject",
+					"Engine",
+					"ComputeFramework",
+				}
+			);
+
+			PrivateDependencyModuleNames.AddRange(
+				new string[]
+				{
+					"ComputeDataInterface",
+					"RenderCore",
+					"RHI",
+				}
+			);
+		}
+	}
+}
diff --git a/Engine/Plugins/Runtime/ComputeFramework/Source/EditableComputeGraph/Internal/ComputeFramework/EditableComputeGraph.h b/Engine/Plugins/Runtime/ComputeFramework/Source/EditableComputeGraph/Internal/ComputeFramework/EditableComputeGraph.h
new file mode 100644
--- /dev/null
+++ b/Engine/Plugins/Runtime/ComputeFramework/Source/EditableComputeGraph/Internal/ComputeFramework/EditableComputeGraph.h
@@ -0,0 +1,159 @@
+// Copyright Epic Games, Inc. All Rights Reserved.
+
+#pragma once
+
+#include "ComputeFramework/ComputeGraph.h"
+#include "EditableComputeGraph.generated.h"
+
+struct FComputeKernelCompileMessage;
+struct FComputeKernelCompileResults;
+
+#define UE_API EDITABLECOMPUTEGRAPH_API
+
+DECLARE_MULTICAST_DELEGATE_OneParam(FOnComputeGraphCompileOutput, TArray<FComputeKernelCompileMessage> const&);
+
+/** Describes one kernel pin: a binding between a kernel HLSL parameter and a data interface function. */
+USTRUCT()
+struct FKernelPin
+{
+	GENERATED_BODY();
+
+	/** Name of the corresponding parameter in the kernel HLSL. */
+	UPROPERTY(EditAnywhere, Category = "Pin", meta = (DisplayName = "Kernel Function"))
+	FString KernelFunctionName;
+	/** Name of a DataInterface entry in the owning FComputeGraphDesc. */
+	UPROPERTY(EditAnywhere, Category = "Pin", meta = (DisplayName = "Data Interface"))
+	FName DataInterfaceName;
+	/** Name of the function exposed by that data interface. */
+	UPROPERTY(EditAnywhere, Category = "Pin", meta = (DisplayName = "Data Interface Function"))
+	FString DataInterfaceFunctionName;
+	
+	/** Transient value that the details customization uses this to display an orphan warning. */
+	UPROPERTY(Transient, VisibleAnywhere, Category = "Pin")
+	bool bOrphaned = false;
+};
+
+/** Describes a single compute kernel: shader source and associated properties. */
+USTRUCT()
+struct FComputeGraphKernelDesc
+{
+	GENERATED_BODY();
+
+	/** Unique name identifying this kernel within the graph. */
+	UPROPERTY(EditAnywhere, Category = "Kernel")
+	FName Name;
+	/** Full HLSL source text for this kernel. */
+	UPROPERTY(EditAnywhere, Category = "Kernel")
+	FString SourceText;
+	/** Name of the HLSL function that serves as the compute shader entry point. */
+	UPROPERTY(EditAnywhere, Category = "Kernel")
+	FString EntryPoint;
+	/** Thread-group dimensions passed to the shader compiler (X * Y * Z threads per group). */
+	UPROPERTY(EditAnywhere, Category = "Kernel")
+	FIntVector GroupSize = FIntVector(64, 1, 1);
+	/** Kernel inputs. Each entry maps a kernel HLSL parameter to a data interface read function. */
+	UPROPERTY(EditAnywhere, Category = "Kernel")
+	TArray<FKernelPin> Inputs;
+	/** Kernel outputs. Each entry maps a kernel HLSL parameter to a data interface write function. */
+	UPROPERTY(EditAnywhere, Category = "Kernel")
+	TArray<FKernelPin> Outputs;
+};
+
+/** Describes a binding object: a named ActorComponent class that data interfaces can reference. */
+USTRUCT()
+struct FComputeGraphDataBindingObjectDesc
+{
+	GENERATED_BODY();
+
+	/** Unique name identifying this binding object within the graph. */
+	UPROPERTY(EditAnywhere, Category = "BindingObject")
+	FName Name;
+	/** The ActorComponent subclass that this binding object wraps. */
+	UPROPERTY(EditAnywhere, Category = "BindingObject", meta = (AllowedClasses = "/Script/Engine.ActorComponent"))
+	TObjectPtr<UClass> Type;
+};
+
+/** Describes a data interface: a named UComputeDataInterface class optionally bound to a binding object. */
+USTRUCT()
+struct FComputeGraphDataInterfaceDesc
+{
+	GENERATED_BODY();
+
+	/** Unique name identifying this data interface within the graph. */
+	UPROPERTY(EditAnywhere, Category = "DataInterface")
+	FName Name;
+	/** Name of the FComputeGraphDataBindingObjectDesc that provides the component instance at runtime. */
+	UPROPERTY(EditAnywhere, Category = "DataInterface", meta = (DisplayName = "Binding Object"))
+	FName BindingObjectName;
+	/** The ComputeDataInterface subclass that implements the read/write functions for this interface. */
+	UPROPERTY(EditAnywhere, Category = "DataInterface", meta = (AllowedClasses = "/Script/ComputeFramework.ComputeDataInterface"))
+	TObjectPtr<UClass> Type;
+	/** Per-instance settings object. This is created automatically when Type is set. */
+	UPROPERTY(VisibleAnywhere, Category = "DataInterface", Instanced, NoClear, meta = (ShowOnlyInnerProperties, EditCondition = "Type!=nullptr", EditConditionHides))
+	TObjectPtr<UComputeDataInterface> Settings;
+};
+
+/** Top-level description of an editable compute graph: its kernels, binding objects, and data interfaces. */
+USTRUCT()
+struct FComputeGraphDesc
+{
+	GENERATED_BODY();
+
+	/** All compute kernels in the graph. */
+	UPROPERTY(EditAnywhere, Category = "Kernel")
+	TArray<FComputeGraphKernelDesc> Kernels;
+	/** All binding objects available for data interfaces to reference. */
+	UPROPERTY(EditAnywhere, Category = "BindingObject")
+	TArray<FComputeGraphDataBindingObjectDesc> BindingObjects;
+	/** All data interfaces that supply read/write functions to kernel pins. */
+	UPROPERTY(EditAnywhere, Category = "DataInterface")
+	TArray<FComputeGraphDataInterfaceDesc> DataInterfaces;
+};
+
+/** Simple data driven compute graph implementation. */
+UCLASS(MinimalAPI, meta = (Experimental))
+class UEditableComputeGraph : public UComputeGraph
+{
+	GENERATED_BODY()
+
+protected:
+#if WITH_EDITORONLY_DATA
+	/** The graph description. */
+	UPROPERTY(EditAnywhere, Category = "Graph", meta = (EditInline))
+	FComputeGraphDesc GraphDescription;
+#endif
+
+#if WITH_EDITOR
+public:
+	/** Returns the graph description (read-only and mutable overloads). */
+	FComputeGraphDesc const& GetGraphDescription() const { return GraphDescription; }
+	FComputeGraphDesc& GetGraphDescription() { return GraphDescription; }
+
+	/** Rebuild the compute graph topology and trigger shader compilation. Called by the asset editor Compile button. */
+	UE_API void RebuildGraph();
+
+	/** Fired on the game thread for each kernel that finishes compiling. */
+	FOnComputeGraphCompileOutput OnCompileOutputChanged;
+
+protected:
+	//~ Begin UObject Interface.
+	void PostInitProperties() override;
+	void PostEditChangeProperty(FPropertyChangedEvent& PropertyChangedEvent) override;
+	void PostTransacted(const FTransactionObjectEvent& TransactionEvent) override;
+	//~ End UObject Interface.
+
+	//~ Begin UComputeGraph Interface.
+	void OnKernelCompilationComplete(int32 InKernelIndex, FComputeKernelCompileResults const& InCompileResults) override;
+	//~ End UComputeGraph Interface.
+
+private:
+	/** Translates GraphDescription into a UComputeGraph topology and submits kernels for compilation. */
+	void Rebuild();
+
+	/** Maps compiled kernel index as seen by OnKernelCompilationComplete() to a kernel display name. Populated by Rebuild(). */
+	TArray<FName> KernelIndexToName;
+
+#endif // WITH_EDITOR
+};
+
+#undef UE_API
diff --git a/Engine/Plugins/Runtime/ComputeFramework/Source/EditableComputeGraph/Private/EditableComputeGraph.cpp b/Engine/Plugins/Runtime/ComputeFramework/Source/EditableComputeGraph/Private/EditableComputeGraph.cpp
new file mode 100644
--- /dev/null
+++ b/Engine/Plugins/Runtime/ComputeFramework/Source/EditableComputeGraph/Private/EditableComputeGraph.cpp
@@ -0,0 +1,207 @@
+// Copyright Epic Games, Inc. All Rights Reserved.
+
+#include "ComputeFramework/EditableComputeGraph.h"
+
+#include "Components/SceneComponent.h"
+#include "ComputeFramework/ComputeDataInterface.h"
+#include "ComputeFramework/ComputeDataInterfaceDispatch.h"
+#include "ComputeFramework/ComputeGraphBuilder.h"
+#include "ComputeFramework/ComputeKernelCompileResult.h"
+#include "ComputeFramework/ComputeKernelSource.h"
+#include "Misc/TransactionObjectEvent.h"
+
+#include UE_INLINE_GENERATED_CPP_BY_NAME(EditableComputeGraph)
+
+#if WITH_EDITOR
+
+void UEditableComputeGraph::PostInitProperties()
+{
+	Super::PostInitProperties();
+
+	if (!HasAnyFlags(RF_ClassDefaultObject))
+	{
+		// Every graph must have a binding object.
+		if (GraphDescription.BindingObjects.Num() == 0)
+		{
+			FComputeGraphDataBindingObjectDesc& BindingObject = GraphDescription.BindingObjects.AddDefaulted_GetRef();
+			BindingObject.Name = FName(TEXT("Binding"));
+			BindingObject.Type = USceneComponent::StaticClass();
+		}
+		// Every graph must have a Dispatch data interface (provides execution + thread count).
+		const bool bHasDispatch = GraphDescription.DataInterfaces.ContainsByPredicate([](FComputeGraphDataInterfaceDesc const& D) { return D.Type == UComputeDataInterfaceDispatch::StaticClass(); });
+		if (!bHasDispatch)
+		{
+			FComputeGraphDataInterfaceDesc& Dispatch = GraphDescription.DataInterfaces.AddDefaulted_GetRef();
+			Dispatch.Name = FName(TEXT("Dispatch"));
+			Dispatch.BindingObjectName = GraphDescription.BindingObjects[0].Name;
+			Dispatch.Type = UComputeDataInterfaceDispatch::StaticClass();
+		}
+	}
+}
+
+void UEditableComputeGraph::PostEditChangeProperty(FPropertyChangedEvent& PropertyChangedEvent)
+{
+	Super::PostEditChangeProperty(PropertyChangedEvent);
+
+	if (PropertyChangedEvent.MemberProperty && PropertyChangedEvent.Property)
+	{
+		if (PropertyChangedEvent.MemberProperty->GetFName() == GET_MEMBER_NAME_CHECKED(UEditableComputeGraph, GraphDescription))
+		{
+			// Skip auto-rebuild for pure name changes. Renaming a kernel, interface or binding object is a label edit that doesn't alter HLSL or graph structure.
+			// The editor applies cross-reference fixups in OnFinishedChangingProperties (which fires after this), so rebuilding here would see stale pin names.
+			// The explicit compile button handles recompilation after fixups are applied.
+			static const FName Name("Name");
+			if (PropertyChangedEvent.Property->GetFName() != Name)
+			{
+				OnCompileOutputChanged.Broadcast({});
+				Rebuild();
+				FString ValidationErrors;
+				if (!ValidateGraph(&ValidationErrors))
+				{
+					FComputeKernelCompileMessage Msg;
+					Msg.Type = FComputeKernelCompileMessage::EMessageType::Error;
+					Msg.Text = ValidationErrors.TrimEnd();
+					OnCompileOutputChanged.Broadcast({ Msg });
+				}
+				else
+				{
+					UpdateResources();
+				}
+			}
+		}
+	}
+}
+
+void UEditableComputeGraph::PostTransacted(const FTransactionObjectEvent& TransactionEvent)
+{
+	Super::PostTransacted(TransactionEvent);
+	if (TransactionEvent.GetEventType() == ETransactionObjectEventType::UndoRedo)
+	{
+		RebuildGraph();
+	}
+}
+
+void UEditableComputeGraph::OnKernelCompilationComplete(int32 InKernelIndex, FComputeKernelCompileResults const& InCompileResults)
+{
+	const FName KernelName = KernelIndexToName.IsValidIndex(InKernelIndex) ? KernelIndexToName[InKernelIndex] : NAME_None;
+	if (KernelName.IsNone() || InCompileResults.Messages.IsEmpty())
+	{
+		OnCompileOutputChanged.Broadcast(InCompileResults.Messages);
+		return;
+	}
+
+	// Add the kernel name to each message.
+	const FString Prefix = FString::Printf(TEXT("[%s] "), *KernelName.ToString());
+	TArray<FComputeKernelCompileMessage> Tagged = InCompileResults.Messages;
+	for (FComputeKernelCompileMessage& Msg : Tagged)
+	{
+		Msg.Text = Prefix + Msg.Text;
+	}
+	
+	OnCompileOutputChanged.Broadcast(Tagged);
+}
+
+void UEditableComputeGraph::RebuildGraph()
+{
+	OnCompileOutputChanged.Broadcast({});
+	Rebuild();
+	FString ValidationErrors;
+	if (!ValidateGraph(&ValidationErrors))
+	{
+		FComputeKernelCompileMessage Msg;
+		Msg.Type = FComputeKernelCompileMessage::EMessageType::Error;
+		Msg.Text = ValidationErrors.TrimEnd();
+		OnCompileOutputChanged.Broadcast({ Msg });
+		return;
+	}
+	UpdateResources();
+}
+
+void UEditableComputeGraph::Rebuild()
+{
+	FComputeGraphBuilder Builder;
+
+	// Create text based kernel sources and add to the builder.
+	KernelIndexToName.Reset();
+	TMap<FName, FKernelHandle> KernelHandleMap;
+	for (FComputeGraphKernelDesc const& KernelDesc : GraphDescription.Kernels)
+	{
+		if (KernelDesc.Name.IsNone() || KernelDesc.SourceText.IsEmpty() || KernelDesc.EntryPoint.IsEmpty())
+		{
+			continue;
+		}
+
+		UComputeKernelSourceWithText* Source = NewObject<UComputeKernelSourceWithText>(this);
+		Source->SourceText = KernelDesc.SourceText;
+		Source->EntryPoint = KernelDesc.EntryPoint;
+		Source->GroupSize  = KernelDesc.GroupSize;
+
+		KernelHandleMap.Add(KernelDesc.Name, Builder.AddKernel(Source, KernelDesc.Name));
+		KernelIndexToName.Add(KernelDesc.Name);
+	}
+
+	// Binding object names.
+	TMap<FName, UClass*> BindingClassMap;
+	for (FComputeGraphDataBindingObjectDesc const& BindingObjectDesc : GraphDescription.BindingObjects)
+	{
+		if (!BindingObjectDesc.Name.IsNone() && BindingObjectDesc.Type)
+		{
+			BindingClassMap.Add(BindingObjectDesc.Name, BindingObjectDesc.Type);
+		}
+	}
+
+	// Data interfaces.
+	TMap<FName, FInterfaceHandle> InterfaceHandleMap;
+	for (FComputeGraphDataInterfaceDesc& InterfaceDesc : GraphDescription.DataInterfaces)
+	{
+		if (InterfaceDesc.BindingObjectName.IsNone() || InterfaceDesc.Type == nullptr)
+		{
+			continue;
+		}
+
+		UClass* BindingClass = BindingClassMap.FindRef(InterfaceDesc.BindingObjectName);
+		if (!BindingClass)
+		{
+			continue;
+		}
+
+		// Preserve any existing settings.
+		if (InterfaceDesc.Type->ChildProperties)
+		{
+			if (!InterfaceDesc.Settings || !InterfaceDesc.Settings->IsA(InterfaceDesc.Type))
+			{
+				InterfaceDesc.Settings = NewObject<UComputeDataInterface>(this, InterfaceDesc.Type);
+			}
+		}
+		else
+		{
+			InterfaceDesc.Settings = nullptr;
+		}
+
+		UComputeDataInterface* DataInterface = InterfaceDesc.Settings
+			? DuplicateObject<UComputeDataInterface>(InterfaceDesc.Settings, this)
+			: NewObject<UComputeDataInterface>(this, InterfaceDesc.Type);
+
+		InterfaceHandleMap.Add(InterfaceDesc.Name, Builder.AddDataInterface(DataInterface, BindingClass, InterfaceDesc.Name));
+	}
+
+	// Connect pins. Invalid handles are silently skipped by the builder.
+	for (FComputeGraphKernelDesc const& KernelDesc : GraphDescription.Kernels)
+	{
+		const FKernelHandle KernelHandle = KernelHandleMap.FindRef(KernelDesc.Name);
+
+		for (FKernelPin const& Pin : KernelDesc.Inputs)
+		{
+			Builder.ConnectInput(KernelHandle, Pin.KernelFunctionName, InterfaceHandleMap.FindRef(Pin.DataInterfaceName), Pin.DataInterfaceFunctionName);
+		}
+
+		for (FKernelPin const& Pin : KernelDesc.Outputs)
+		{
+			Builder.ConnectOutput(KernelHandle, Pin.KernelFunctionName, InterfaceHandleMap.FindRef(Pin.DataInterfaceName), Pin.DataInterfaceFunctionName);
+		}
+	}
+
+	Builder.Build(*this, *this);
+}
+
+#endif
diff --git a/Engine/Plugins/Runtime/ComputeFramework/Source/EditableComputeGraph/Private/EditableComputeGraphModule.cpp b/Engine/Plugins/Runtime/ComputeFramework/Source/EditableComputeGraph/Private/EditableComputeGraphModule.cpp
new file mode 100644
--- /dev/null
+++ b/Engine/Plugins/Runtime/ComputeFramework/Source/EditableComputeGraph/Private/EditableComputeGraphModule.cpp
@@ -0,0 +1,5 @@
+// Copyright Epic Games, Inc. All Rights Reserved.
+
+#include "Modules/ModuleManager.h"
+
+IMPLEMENT_MODULE(FDefaultModuleImpl, EditableComputeGraph)
diff --git a/Engine/Plugins/Runtime/ComputeFramework/Source/EditableComputeGraphEditor/EditableComputeGraphEditor.Build.cs b/Engine/Plugins/Runtime/ComputeFramework/Source/EditableComputeGraphEditor/EditableComputeGraphEditor.Build.cs
new file mode 100644
--- /dev/null
+++ b/Engine/Plugins/Runtime/ComputeFramework/Source/EditableComputeGraphEditor/EditableComputeGraphEditor.Build.cs
@@ -0,0 +1,29 @@
+// Copyright Epic Games, Inc. All Rights Reserved.
+
+namespace UnrealBuildTool.Rules
+{
+	public class EditableComputeGraphEditor : ModuleRules
+	{
+		public EditableComputeGraphEditor(ReadOnlyTargetRules Target) : base(Target)
+		{
+			PrivateDependencyModuleNames.AddRange(
+				new string[]
+				{
+					"AssetDefinition",
+					"ComputeDataInterface",
+					"ComputeFramework",
+					"EditableComputeGraph",
+					"Core",
+					"CoreUObject",
+					"Engine",
+					"InputCore",
+					"PropertyEditor",
+					"Slate",
+					"SlateCore",
+					"ToolMenus",
+					"UnrealEd",
+				}
+			);
+		}
+	}
+}
diff --git a/Engine/Plugins/Runtime/ComputeFramework/Source/EditableComputeGraphEditor/Private/ComputeFramework/ComputeGraphHlslEditor.h b/Engine/Plugins/Runtime/ComputeFramework/Source/EditableComputeGraphEditor/Private/ComputeFramework/ComputeGraphHlslEditor.h
new file mode 100644
--- /dev/null
+++ b/Engine/Plugins/Runtime/ComputeFramework/Source/EditableComputeGraphEditor/Private/ComputeFramework/ComputeGraphHlslEditor.h
@@ -0,0 +1,72 @@
+// Copyright Epic Games, Inc. All Rights Reserved.
+
+#pragma once
+
+#include "Widgets/SCompoundWidget.h"
+
+class FComputeGraphHlslSyntaxHighlighter;
+class SMultiLineEditableTextBox;
+class STextBlock;
+class UEditableComputeGraph;
+
+DECLARE_DELEGATE_TwoParams(FOnComputeGraphHlslTextCommitted, FName /* KernelName */, FString const& /* NewText */);
+DECLARE_DELEGATE(FOnComputeGraphHlslTextChanged);
+DECLARE_DELEGATE(FOnComputeGraphHlslCompileRequested);
+
+/**
+ * HLSL editor for the EditableComputeGraph asset editor.
+ * Displays the source text of the currently selected kernel in a monospace multi-line text box. 
+ * User commits the text by changing focus or F8 to compile.
+ */
+class SComputeGraphHlslEditor : public SCompoundWidget
+{
+public:
+	SLATE_BEGIN_ARGS(SComputeGraphHlslEditor) {}
+		SLATE_ARGUMENT(UEditableComputeGraph*, Asset)
+		SLATE_EVENT(FOnComputeGraphHlslTextCommitted, OnTextCommitted_WithKernel)
+		SLATE_EVENT(FOnComputeGraphHlslTextChanged, OnTextChanged)
+		SLATE_EVENT(FOnComputeGraphHlslCompileRequested, OnCompileRequested)
+	SLATE_END_ARGS()
+
+	void Construct(FArguments const& InArgs);
+
+	/** Switch the editor to display a different kernel's source text. */
+	void SetActiveKernel(FName KernelName);
+
+	/** Highlight the given DI_ function names with the bound-function colour. */
+	void SetBoundFunctions(TArray<FString> const& Names);
+
+	/** Flush any in-progress edit. */
+	void CommitCurrentText();
+
+private:
+	/** Override key down handler to intercept F8 for triggering a compile request. */
+	FReply OnKeyDown(FGeometry const& MyGeometry, FKeyEvent const& InKeyEvent) override;
+	/** First chance key handler wired to the text box to intercept Tab before Slate focus traversal. */
+	FReply OnTextBoxKeyDown(FGeometry const& MyGeometry, FKeyEvent const& InKeyEvent);
+
+	/** Fires the commit delegate when the user finishes editing. */
+	void OnTextCommitted(FText const& NewText, ETextCommit::Type CommitType);
+	/** Notifies the toolkit that the text has changed so it can mark the asset dirty. */
+	void OnTextChangedInternal(FText const& NewText);
+
+	/** The asset being edited. */
+	TWeakObjectPtr<UEditableComputeGraph> Asset;
+
+	/** Fires with the kernel name and new source text when the user commits an edit. */
+	FOnComputeGraphHlslTextCommitted OnTextCommitted_WithKernel;
+	/** Fires whenever the text box content changes. */
+	FOnComputeGraphHlslTextChanged OnTextChanged;
+	/** Fires when the user requests a compile (F8). */
+	FOnComputeGraphHlslCompileRequested OnCompileRequested;
+
+	/** Name of the kernel whose source text is currently displayed. */
+	FName ActiveKernelName;
+	/** The multi-line text box used for HLSL input. */
+	TSharedPtr<SMultiLineEditableTextBox> TextBox;
+	/** Shown in place of the text box when no kernel is selected. */
+	TSharedPtr<STextBlock> PlaceholderText;
+
+	/** Syntax highlighter that colours keywords, DI_ calls, and bound functions. */
+	TSharedPtr<FComputeGraphHlslSyntaxHighlighter> SyntaxHighlighter;
+};
diff --git a/Engine/Plugins/Runtime/ComputeFramework/Source/EditableComputeGraphEditor/Private/ComputeFramework/ComputeGraphHlslSyntaxHighlighter.h b/Engine/Plugins/Runtime/ComputeFramework/Source/EditableComputeGraphEditor/Private/ComputeFramework/ComputeGraphHlslSyntaxHighlighter.h
new file mode 100644
--- /dev/null
+++ b/Engine/Plugins/Runtime/ComputeFramework/Source/EditableComputeGraphEditor/Private/ComputeFramework/ComputeGraphHlslSyntaxHighlighter.h
@@ -0,0 +1,40 @@
+// Copyright Epic Games, Inc. All Rights Reserved.
+
+#pragma once
+
+#include "CoreMinimal.h"
+#include "Text/HLSLSyntaxHighlighterMarshaller.h"
+
+/**
+ * HLSL syntax highlighter for the EditableComputeGraph kernel editor.
+ * Extends the base HLSL highlighter with an extra colour for DI_ function names that are currently bound in the kernel's pin list.
+ */
+class FComputeGraphHlslSyntaxHighlighter : public FHLSLSyntaxHighlighterMarshaller
+{
+public:
+	/** Creates the highlighter. */
+	static TSharedRef<FComputeGraphHlslSyntaxHighlighter> Create(FSyntaxTextStyle const& InSyntaxTextStyle, FTextBlockStyle const& InBoundFunctionStyle);
+
+	/** Update the set of bound DI_ function names and trigger a re-highlight. */
+	void SetBoundFunctions(TArray<FString> const& Names);
+
+protected:
+	/** Applies bound-function colouring on top of the base HLSL token styles. */
+	FTextLayout::FNewLineData ProcessTokenizedLine(
+		ISyntaxTokenizer::FTokenizedLine const& TokenizedLine,
+		int32 const& LineNumber,
+		FString const& SourceString,
+		EParseState& ParseState) override;
+
+private:
+	/** Private constructor. Use Create(). */
+	FComputeGraphHlslSyntaxHighlighter(
+		TSharedPtr<ISyntaxTokenizer> InTokenizer,
+		FSyntaxTextStyle const& InSyntaxTextStyle,
+		FTextBlockStyle const& InBoundFunctionStyle);
+
+	/** Text style applied to tokens that match a name in BoundFunctionNames. */
+	FTextBlockStyle BoundFunctionStyle;
+	/** DI_ function names currently wired to a kernel pin. */
+	TSet<FString> BoundFunctionNames;
+};
diff --git a/Engine/Plugins/Runtime/ComputeFramework/Source/EditableComputeGraphEditor/Private/ComputeFramework/ComputeGraphNavigator.h b/Engine/Plugins/Runtime/ComputeFramework/Source/EditableComputeGraphEditor/Private/ComputeFramework/ComputeGraphNavigator.h
new file mode 100644
--- /dev/null
+++ b/Engine/Plugins/Runtime/ComputeFramework/Source/EditableComputeGraphEditor/Private/ComputeFramework/ComputeGraphNavigator.h
@@ -0,0 +1,121 @@
+// Copyright Epic Games, Inc. All Rights Reserved.
+
+#pragma once
+
+#include "CoreMinimal.h"
+#include "ComputeFramework/EditableComputeGraphDetailCustomization.h"
+#include "Widgets/SCompoundWidget.h"
+#include "Widgets/Views/STreeView.h"
+
+class UEditableComputeGraph;
+struct FNavigatorItem;
+using FNavigatorItemPtr = TSharedPtr<FNavigatorItem>;
+
+/** The category of item shown in the compute graph toolkit navigator panel. */
+enum class EComputeGraphItemKind : uint8
+{
+	None,
+	Kernel,
+	Interface,
+	BindingObject,
+};
+
+/** Shared selection state between the navigator and the toolkit detail customization. */
+struct FComputeGraphEditorSelection
+{
+	/** Item type. */
+	EComputeGraphItemKind Kind = EComputeGraphItemKind::None;
+	/** Index of the item within its kind array. */
+	int32 Index = INDEX_NONE;
+};
+
+/** Fired when the user clicks a leaf item. Index is the position in the relevant array. */
+DECLARE_DELEGATE_ThreeParams(FOnComputeGraphItemSelected, EComputeGraphItemKind, int32 /* Index */, FName /* Name */);
+
+/** Fired when the user clicks a section [+] button. */
+DECLARE_DELEGATE_OneParam(FOnComputeGraphAddItem, EComputeGraphItemKind);
+
+/** Fired when the user selects Delete from the context menu. */
+DECLARE_DELEGATE_TwoParams(FOnComputeGraphDeleteItem, EComputeGraphItemKind, int32 /* Index */);
+
+/** Fired when the user selects Duplicate from the context menu. */
+DECLARE_DELEGATE_TwoParams(FOnComputeGraphDuplicateItem, EComputeGraphItemKind, int32 /* Index */);
+
+/** Fired when the user commits a rename (inline text or context menu → Rename). */
+DECLARE_DELEGATE_ThreeParams(FOnComputeGraphRenameItem, EComputeGraphItemKind, int32 /* Index */, FName /* NewName */);
+
+/** Kept for HLSL editor wiring: fires when a kernel is selected. */
+DECLARE_DELEGATE_OneParam(FOnComputeGraphKernelSelected, FName /* KernelName */);
+
+/**
+ * Item navigator for the EditableComputeGraph asset editor.
+ * Supports management of Binding Objects, Data Interfaces and Kernels.
+ * Styled to match the Blueprint item panel.
+ */
+class SComputeGraphNavigator : public SCompoundWidget
+{
+public:
+	SLATE_BEGIN_ARGS(SComputeGraphNavigator) {}
+		SLATE_ARGUMENT(UEditableComputeGraph*, Asset)
+		SLATE_EVENT(FOnComputeGraphItemSelected, OnItemSelected)
+		SLATE_EVENT(FOnComputeGraphAddItem, OnAddItem)
+		SLATE_EVENT(FOnComputeGraphDeleteItem, OnDeleteItem)
+		SLATE_EVENT(FOnComputeGraphDuplicateItem, OnDuplicateItem)
+		SLATE_EVENT(FOnComputeGraphRenameItem, OnRenameItem)
+		SLATE_EVENT(FOnComputeGraphKernelSelected, OnKernelSelected)
+	SLATE_END_ARGS()
+
+	void Construct(FArguments const& InArgs);
+
+	/** Rebuild the tree from the current asset description and reapply selection/expansion. */
+	void Refresh();
+
+	/** Programmatically select an item (e.g. after Add or Duplicate). Pass INDEX_NONE to clear. */
+	void SetSelectedItem(EComputeGraphItemKind Kind, int32 Index);
+
+private:
+	/** Rebuilds RootItems from the current asset description. */
+	void RebuildItems();
+	/** Builds either a section header row or a leaf item row for the tree view. */
+	TSharedRef<ITableRow> GenerateRow(FNavigatorItemPtr Item, TSharedRef<STableViewBase> const& OwnerTable);
+	/** Fills OutChildren from the item's Children array (header items only). */
+	void OnGetChildren(FNavigatorItemPtr Item, TArray<FNavigatorItemPtr>& OutChildren);
+	/** Builds the right-click context menu for a selected leaf item. */
+	TSharedPtr<SWidget> BuildContextMenu();
+
+	/** Override key down handler to intercept F2 for inline rename on the selected item. */
+	FReply OnKeyDown(FGeometry const& MyGeometry, FKeyEvent const& InKeyEvent) override;
+
+	// Delegate functions.
+	void OnSelectionChanged(FNavigatorItemPtr Item, ESelectInfo::Type SelectInfo);
+	void OnExpansionChanged(FNavigatorItemPtr Item, bool bExpanded);
+	void OnAddItemClicked(EComputeGraphItemKind Kind);
+	void OnRenameItemCommitted(EComputeGraphItemKind Kind, int32 Index, FName NewName);
+	void OnDeleteSelected();
+	void OnDuplicateSelected();
+	void OnRenameSelected();
+
+	/** Returns false and sets OutError if the candidate name is empty or already in use. */
+	bool VerifyRenameText(FText const& NewText, FText& OutError, EComputeGraphItemKind Kind, int32 Index) const;
+
+	/** The asset being navigated; held weakly since this widget is not a GC object. */
+	TWeakObjectPtr<UEditableComputeGraph> Asset;
+
+	/** The tree view widget. */
+	TSharedPtr<STreeView<FNavigatorItemPtr>> TreeView;
+	/** Root items — one per section header; each carries its leaf children. */
+	TArray<FNavigatorItemPtr> RootItems;
+	/** Sections that have been manually collapsed; all others are expanded. */
+	TSet<EComputeGraphItemKind> CollapsedSections;
+
+	/** Current selection; mirrored here so the context menu can read it without a tree query. */
+	FComputeGraphEditorSelection Selection;
+
+	// Delegate objects.
+	FOnComputeGraphItemSelected OnItemSelected;
+	FOnComputeGraphAddItem OnAddItem;
+	FOnComputeGraphDeleteItem OnDeleteItem;
+	FOnComputeGraphDuplicateItem OnDuplicateItem;
+	FOnComputeGraphRenameItem OnRenameItem;
+	FOnComputeGraphKernelSelected OnKernelSelected;
+};
diff --git a/Engine/Plugins/Runtime/ComputeFramework/Source/EditableComputeGraphEditor/Private/ComputeFramework/ComputeGraphView.h b/Engine/Plugins/Runtime/ComputeFramework/Source/EditableComputeGraphEditor/Private/ComputeFramework/ComputeGraphView.h
new file mode 100644
--- /dev/null
+++ b/Engine/Plugins/Runtime/ComputeFramework/Source/EditableComputeGraphEditor/Private/ComputeFramework/ComputeGraphView.h
@@ -0,0 +1,90 @@
+// Copyright Epic Games, Inc. All Rights Reserved.
+
+#pragma once
+
+#include "Widgets/SLeafWidget.h"
+
+class UEditableComputeGraph;
+
+DECLARE_DELEGATE_TwoParams(FOnComputeGraphNodeClicked, FName /*NodeName*/, bool /*bIsKernel*/);
+
+/**
+ * Read only graph view for the EditableComputeGraph asset editor.
+ * Each connected pin pair is drawn as a spline.
+ * The selected kernels and data interfaces can be highlighted along with a dimmer highlight on their neighboring nodes.
+ * Clicking on nodes in the graph view will select them in the toolkit.
+ */
+class SComputeGraphView : public SLeafWidget
+{
+public:
+	SLATE_BEGIN_ARGS(SComputeGraphView) {}
+		SLATE_ARGUMENT(UEditableComputeGraph*, Asset)
+		SLATE_EVENT(FOnComputeGraphNodeClicked, OnNodeClicked)
+	SLATE_END_ARGS()
+
+	void Construct(FArguments const& InArgs);
+
+	/** Highlight edges connected to a specific kernel. */
+	void SetHighlightedKernel(FName KernelName);
+
+	/** Highlight a specific data interface node. */
+	void SetHighlightedInterface(FName InterfaceName);
+
+	/** Rebuild the visual layout from the current asset description. */
+	void Refresh();
+
+protected:
+	//~ Begin SWidget interface.
+	int32 OnPaint(FPaintArgs const& Args, FGeometry const& AllottedGeometry, FSlateRect const& MyCullingRect, FSlateWindowElementList& OutDrawElements, int32 LayerId, FWidgetStyle const& InWidgetStyle, bool bParentEnabled) const override;
+	FVector2D ComputeDesiredSize(float LayoutScaleMultiplier) const override;
+	FReply OnMouseButtonDown(FGeometry const& MyGeometry, FPointerEvent const& MouseEvent) override;
+	//~ End SWidget interface.
+
+private:
+	/** Rebuilds NodeLayouts, Edges, centre lookup maps, and cached size metrics from the current asset description. */
+	void RebuildLayout();
+
+	/** Position and rank of one node in the graph view. */
+	struct FNodeLayout
+	{
+		FName Name;
+		FVector2D Centre; // Local widget space (pixels, pre-DPI-scale).
+		int32 Rank = 0; // Column index: 0 = leftmost interfaces, increasing toward outputs.
+		bool bKernel = false;
+	};
+
+	/** One directed connection between a kernel pin and a data interface. */
+	struct FEdge
+	{
+		FName KernelName;
+		FName InterfaceName;
+		FString KernelFn;
+		FString InterfaceFn;
+		bool bOrphaned = false; // True when the pin has no valid binding; drawn red.
+		bool bIsOutput = false; // True for kernel output pins (kernel → interface direction).
+	};
+
+	/** The asset being visualised; held weakly since this widget is not a GC object. */
+	TWeakObjectPtr<UEditableComputeGraph> Asset;
+
+	/** Flat list of node positions, rebuilt by RebuildLayout(). */
+	TArray<FNodeLayout> NodeLayouts;
+	/** Flat list of edges (one per kernel pin), rebuilt by RebuildLayout(). */
+	TArray<FEdge> Edges;
+	/** Centre position lookup for interface nodes, keyed by name. Rebuilt by RebuildLayout(). */
+	TMap<FName, FVector2D> InterfaceCentreMap;
+	/** Centre position lookup for kernel nodes, keyed by name. Rebuilt by RebuildLayout(). */
+	TMap<FName, FVector2D> KernelCentreMap;
+
+	/** Name of the kernel node to draw highlighted. NAME_None when nothing is selected. */
+	FName HighlightedKernel;
+	/** Name of the interface node to draw highlighted. NAME_None when nothing is selected. */
+	FName HighlightedInterface;
+	/** Highest rank value present in NodeLayouts that drives the widget width. */
+	int32 MaxRank = 0;
+	/** Largest number of nodes sharing a rank column that drives the widget height. */
+	int32 MaxNodesInAnyRank = 0;
+
+	/** Fired when the user clicks a node. */
+	FOnComputeGraphNodeClicked OnNodeClicked;
+};
diff --git a/Engine/Plugins/Runtime/ComputeFramework/Source/EditableComputeGraphEditor/Private/ComputeFramework/EditableComputeGraphAssetDefinition.h b/Engine/Plugins/Runtime/ComputeFramework/Source/EditableComputeGraphEditor/Private/ComputeFramework/EditableComputeGraphAssetDefinition.h
new file mode 100644
--- /dev/null
+++ b/Engine/Plugins/Runtime/ComputeFramework/Source/EditableComputeGraphEditor/Private/ComputeFramework/EditableComputeGraphAssetDefinition.h
@@ -0,0 +1,21 @@
+// Copyright Epic Games, Inc. All Rights Reserved.
+
+#pragma once
+
+#include "AssetDefinitionDefault.h"
+
+#include "EditableComputeGraphAssetDefinition.generated.h"
+
+UCLASS()
+class UAssetDefinition_EditableComputeGraph : public UAssetDefinitionDefault
+{
+	GENERATED_BODY()
+
+	//~ Begin UAssetDefinition Interface.
+	virtual FText GetAssetDisplayName() const override;
+	virtual FLinearColor GetAssetColor() const override;
+	virtual TSoftClassPtr<UObject> GetAssetClass() const override;
+	virtual TConstArrayView<FAssetCategoryPath> GetAssetCategories() const override;
+	virtual EAssetCommandResult OpenAssets(const FAssetOpenArgs& OpenArgs) const override;
+	//~ End UAssetDefinition Interface.
+};
diff --git a/Engine/Plugins/Runtime/ComputeFramework/Source/EditableComputeGraphEditor/Private/ComputeFramework/EditableComputeGraphDetailCustomization.h b/Engine/Plugins/Runtime/ComputeFramework/Source/EditableComputeGraphEditor/Private/ComputeFramework/EditableComputeGraphDetailCustomization.h
new file mode 100644
--- /dev/null
+++ b/Engine/Plugins/Runtime/ComputeFramework/Source/EditableComputeGraphEditor/Private/ComputeFramework/EditableComputeGraphDetailCustomization.h
@@ -0,0 +1,25 @@
+// Copyright Epic Games, Inc. All Rights Reserved.
+
+#pragma once
+
+#include "IDetailCustomization.h"
+
+struct FComputeGraphEditorSelection;
+
+/**
+ * Detail Customization for UEditableComputeGraph.
+ * Hides the full "Graph" category and instead exposes only the current array element that the user has selected in the navigator, using property handle navigation.
+ */
+class FEditableComputeGraphDetailCustomization : public IDetailCustomization
+{
+public:
+	static TSharedRef<IDetailCustomization> MakeInstance(TSharedPtr<FComputeGraphEditorSelection> InSelection);
+
+protected:
+	//~ Begin IDetailCustomization Interface.
+	void CustomizeDetails(IDetailLayoutBuilder& DetailBuilder) override;
+	//~ End IDetailCustomization Interface.
+
+	/** Current selection. */
+	TSharedPtr<FComputeGraphEditorSelection> Selection;
+};
diff --git a/Engine/Plugins/Runtime/ComputeFramework/Source/EditableComputeGraphEditor/Private/ComputeFramework/EditableComputeGraphEditorToolkit.h b/Engine/Plugins/Runtime/ComputeFramework/Source/EditableComputeGraphEditor/Private/ComputeFramework/EditableComputeGraphEditorToolkit.h
new file mode 100644
--- /dev/null
+++ b/Engine/Plugins/Runtime/ComputeFramework/Source/EditableComputeGraphEditor/Private/ComputeFramework/EditableComputeGraphEditorToolkit.h
@@ -0,0 +1,132 @@
+// Copyright Epic Games, Inc. All Rights Reserved.
+
+#pragma once
+
+#include "Textures/SlateIcon.h"
+#include "Toolkits/AssetEditorToolkit.h"
+#include "UObject/GCObject.h"
+#include "Widgets/Views/SListView.h"
+
+enum class EComputeGraphItemKind : uint8;
+struct FComputeGraphEditorSelection;
+struct FComputeKernelCompileMessage;
+class IDetailsView;
+class SComputeGraphHlslEditor;
+class SComputeGraphNavigator;
+class SComputeGraphView;
+class UEditableComputeGraph;
+
+/** Custom asset editor for UEditableComputeGraph. */
+class FEditableComputeGraphEditorToolkit : public FAssetEditorToolkit, public FGCObject
+{
+public:
+	static TSharedRef<FEditableComputeGraphEditorToolkit> Create(UEditableComputeGraph* InAsset, EToolkitMode::Type InMode,	TSharedPtr<IToolkitHost> InToolkitHost);
+
+	/** Returns the asset being edited. */
+	UEditableComputeGraph* GetAsset() const { return Asset; }
+
+	/** Flushes any uncommitted HLSL text then triggers a full graph rebuild and shader compile. */
+	void CompileGraph();
+	/** Whether the graph is compiled and up to date, needs recompile, or has errors. */
+	enum class ECompileState { UpToDate, NeedsCompile, Broken };
+	/** Inspects the graph description to determine the current compile state without recompiling. */
+	ECompileState GetCompileState() const;
+	/** Returns the toolbar icon matching the current compile state. */
+	FSlateIcon GetCompileStatusIcon() const;
+
+protected:
+	//~ Begin FAssetEditorToolkit Interface.
+	FName GetToolkitFName() const override;
+	FText GetBaseToolkitName() const override;
+	FString GetWorldCentricTabPrefix() const override;
+	FLinearColor GetWorldCentricTabColorScale() const override;
+	void RegisterTabSpawners(TSharedRef<FTabManager> const& InTabManager) override;
+	void UnregisterTabSpawners(TSharedRef<FTabManager> const& InTabManager) override;
+	//~ End FAssetEditorToolkit Interface.
+
+	//~ Begin FGCObject Interface.
+	void AddReferencedObjects(FReferenceCollector& Collector) override;
+	FString GetReferencerName() const override;
+	//~ End FGCObject Interface.
+
+	/** Updates the details panel and HLSL editor to reflect the newly selected item. */
+	void OnItemSelected(EComputeGraphItemKind Kind, int32 Index, FName Name);
+	/** Appends a new kernel, interface, or binding object with a generated unique name. */
+	void OnAddItem(EComputeGraphItemKind Kind);
+	/** Removes the item at Index and clears any dangling pin references to it. */
+	void OnDeleteItem(EComputeGraphItemKind Kind, int32 Index);
+	/** Appends a deep copy of the item at Index with a "_Copy" suffix on its name. */
+	void OnDuplicateItem(EComputeGraphItemKind Kind, int32 Index);
+	/** Renames the item at Index and propagates the new name to any referencing pins. */
+	void OnRenameItem(EComputeGraphItemKind Kind, int32 Index, FName NewName);
+	/** Loads the selected kernel's source into the HLSL editor. */
+	void OnKernelSelected(FName KernelName);
+	/** Syncs pin lists from the new HLSL text and marks the graph dirty. */
+	void OnHlslTextCommitted(FName KernelName, FString const& NewText);
+	/** Selects the corresponding kernel or interface in the navigator when a graph node is clicked. */
+	void OnGraphNodeClicked(FName NodeName, bool bIsKernel);
+	/** Appends newly arrived compile messages to the output list view. */
+	void OnCompileOutputChanged(TArray<FComputeKernelCompileMessage> const& NewMessages);
+	/** Navigates the navigator to the kernel associated with the clicked output message. */
+	void OnOutputSelectionChanged(TSharedPtr<FComputeKernelCompileMessage> Selected, ESelectInfo::Type SelectInfo);
+
+private:
+	/** Initialises the toolkit: creates the layout, registers tab spawners, and binds all callbacks. */
+	void Init(UEditableComputeGraph* InAsset, EToolkitMode::Type InMode, TSharedPtr<IToolkitHost> InToolkitHost);
+
+	// Tab spawn functions.
+	TSharedRef<SDockTab> SpawnTab_Navigator(FSpawnTabArgs const& Args);
+	TSharedRef<SDockTab> SpawnTab_HlslEditor(FSpawnTabArgs const& Args);
+	TSharedRef<SDockTab> SpawnTab_GraphView(FSpawnTabArgs const& Args);
+	TSharedRef<SDockTab> SpawnTab_Details(FSpawnTabArgs const& Args);
+	TSharedRef<SDockTab> SpawnTab_Output(FSpawnTabArgs const& Args);
+
+	/** Builds one row in the compile output list view. */
+	TSharedRef<ITableRow> GenerateOutputRow(TSharedPtr<FComputeKernelCompileMessage> Message, TSharedRef<STableViewBase> const& OwnerTable);
+
+	/** Adds the compile button and status icon to the asset editor toolbar. */
+	void ExtendToolbar();
+	/** Pushes the current kernel's bound pin names into the HLSL editor for syntax highlighting. */
+	void RefreshHlslBoundFunctions(FName KernelName);
+	/** Marks the graph dirty, refreshes the navigator, graph view, and details panel. */
+	void RefreshNavigatorAndDetails();
+
+	/** Unique name helper. Returns "Base", "Base1", "Base2" etc while avoiding names in Existing. */
+	static FName GenerateUniqueName(TArray<FName> const& Existing, FStringView Base);
+
+private:
+	/** The asset being edited. */
+	TObjectPtr<UEditableComputeGraph> Asset;
+
+	/** Tree-view panel listing kernels, interfaces, and binding objects. */
+	TSharedPtr<SComputeGraphNavigator> NavigatorWidget;
+	/** Multi-line HLSL editor panel. */
+	TSharedPtr<SComputeGraphHlslEditor> HlslEditorWidget;
+	/** Bipartite graph visualisation panel. */
+	TSharedPtr<SComputeGraphView> GraphViewWidget;
+	/** Details panel showing per-item properties. */
+	TSharedPtr<IDetailsView> DetailsView;
+
+	/** The list view widget in the Output tab. */
+	TSharedPtr<SListView<TSharedPtr<FComputeKernelCompileMessage>>> OutputListView;
+	/** Messages currently displayed in the Output tab; repopulated on each compile. */
+	TArray<TSharedPtr<FComputeKernelCompileMessage>> OutputMessages;
+	/** Bound to UEditableComputeGraph::OnCompileOutputChanged to refresh the output tab. */
+	FDelegateHandle CompileOutputHandle;
+
+	/** Shared with FEditableComputeGraphDetailCustomization; updated on every selection change. */
+	TSharedPtr<FComputeGraphEditorSelection> EditorSelection;
+
+	/** True whenever graph data has been mutated since the last successful compile. */
+	bool bDirty = true;
+
+	/** Name of the currently selected Interface or BindingObject, captured at selection time so OnFinishedChangingProperties can detect in-panel renames and apply cross-reference fixups. */
+	FName CachedSelectedItemName = NAME_None;
+
+	// Tab Id FNames.
+	static FName const TabId_Navigator;
+	static FName const TabId_HlslEditor;
+	static FName const TabId_GraphView;
+	static FName const TabId_Details;
+	static FName const TabId_Output;
+};
diff --git a/Engine/Plugins/Runtime/ComputeFramework/Source/EditableComputeGraphEditor/Private/ComputeFramework/EditableComputeGraphFactory.h b/Engine/Plugins/Runtime/ComputeFramework/Source/EditableComputeGraphEditor/Private/ComputeFramework/EditableComputeGraphFactory.h
new file mode 100644
--- /dev/null
+++ b/Engine/Plugins/Runtime/ComputeFramework/Source/EditableComputeGraphEditor/Private/ComputeFramework/EditableComputeGraphFactory.h
@@ -0,0 +1,18 @@
+// Copyright Epic Games, Inc. All Rights Reserved.
+
+#pragma once
+
+#include "Factories/Factory.h"
+#include "EditableComputeGraphFactory.generated.h"
+
+UCLASS(hidecategories = Object)
+class UEditableComputeGraphFactory : public UFactory
+{
+	GENERATED_BODY()
+
+	UEditableComputeGraphFactory();
+
+	//~ Begin UFactory Interface.
+	virtual UObject* FactoryCreateNew(UClass* InClass, UObject* InParent, FName InName, EObjectFlags Flags, UObject* Context, FFeedbackContext* Warn) override;
+	//~ End UFactory Interface.
+};
diff --git a/Engine/Plugins/Runtime/ComputeFramework/Source/EditableComputeGraphEditor/Private/ComputeFramework/HlslExternalPinParser.h b/Engine/Plugins/Runtime/ComputeFramework/Source/EditableComputeGraphEditor/Private/ComputeFramework/HlslExternalPinParser.h
new file mode 100644
--- /dev/null
+++ b/Engine/Plugins/Runtime/ComputeFramework/Source/EditableComputeGraphEditor/Private/ComputeFramework/HlslExternalPinParser.h
@@ -0,0 +1,40 @@
+// Copyright Epic Games, Inc. All Rights Reserved.
+
+#pragma once
+
+#include "CoreMinimal.h"
+
+/** Helper functions for finding function bindings in HLSL source. */
+class FHlslExternalPinParser
+{
+public:
+	/** Result of one DI_ call site found by FindExternalPins(). */
+	struct FPinDeclaration
+	{
+		/** Full function name including the DI_ prefix, e.g. "DI_ReadValue". */
+		FString FunctionName;
+		/** True when the call is considered to be an output. */
+		bool bIsOutput = false;
+	};
+
+	/**
+	 * Scan SourceText for DataInterface function calls.
+	 * The convention used by UEditableComputeGraph kernel authoring is that DataInterface function calls have the "DI_" prefix in the kernel.
+	 * These will be presented in the toolkit UI for binding.
+	 * The input/output direction is inferred from call-site context. If DI_*() is a standalone statement (only whitespace precedes it on the same line) then
+	 * the call discards its return value and we make the assumption it is an output pin.
+	 */
+	static TArray<FPinDeclaration> FindExternalPins(FStringView SourceText);
+
+	/**
+	 * Scan SourceText for top level function definitions. These are signatures followed by a body.
+	 * We returns the bare function names in declaration order.
+	 * This is usful for finding potential kernel EntryPoints.
+	 */
+	static TArray<FString> FindFunctionDefinitions(FStringView SourceText);
+
+	/**
+	 * Returns true if FunctionName appears in SourceText as a whole-word identifier followed by '('.
+	 */
+	static bool FunctionExistsInText(FStringView SourceText, FString const& FunctionName);
+};
diff --git a/Engine/Plugins/Runtime/ComputeFramework/Source/EditableComputeGraphEditor/Private/ComputeFramework/KernelPinCustomization.h b/Engine/Plugins/Runtime/ComputeFramework/Source/EditableComputeGraphEditor/Private/ComputeFramework/KernelPinCustomization.h
new file mode 100644
--- /dev/null
+++ b/Engine/Plugins/Runtime/ComputeFramework/Source/EditableComputeGraphEditor/Private/ComputeFramework/KernelPinCustomization.h
@@ -0,0 +1,79 @@
+// Copyright Epic Games, Inc. All Rights Reserved.
+
+#pragma once
+
+#include "IPropertyTypeCustomization.h"
+#include "Styling/SlateColor.h"
+
+class IPropertyHandle;
+
+/**
+ * Details panel customization for FKernelPin.
+ * Holds combo boxes for populating the DataInterface and function.
+ */
+class FKernelPinCustomization : public IPropertyTypeCustomization
+{
+public:
+	static TSharedRef<IPropertyTypeCustomization> MakeInstance();
+
+protected:
+	//~ Begin IPropertyTypeCustomization Interface.
+	void CustomizeHeader(TSharedRef<IPropertyHandle> StructPropertyHandle, FDetailWidgetRow& HeaderRow, IPropertyTypeCustomizationUtils& CustomizationUtils) override;
+	void CustomizeChildren(TSharedRef<IPropertyHandle> StructPropertyHandle, IDetailChildrenBuilder& ChildBuilder, IPropertyTypeCustomizationUtils& CustomizationUtils) override;
+	//~ End IPropertyTypeCustomization Interface.
+
+private:
+	/** Rebuilds the kernel function combo options from the owning kernel's HLSL source. */
+	void RefreshKernelFunctionNameOptions();
+	/** Rebuilds the data interface combo options from the graph's interface list. */
+	void RefreshInterfaceNameOptions();
+	/** Rebuilds the DI function combo options from the currently selected interface's declared functions. */
+	void RefreshFunctionNameOptions();
+
+	/** Returns true when the pin is orphaned or any required field is unset. */
+	bool NeedsAttention() const;
+	/** Text shown in the collapsed header row: "KernelFn → Interface::DIFn". */
+	FText GetHeaderSummaryText() const;
+	/** Orange when NeedsAttention(), foreground colour otherwise. */
+	FSlateColor GetHeaderSummaryColor() const;
+	/** Tooltip explaining why the pin needs attention (orphaned or incomplete). */
+	FText GetOrphanedTooltipText() const;
+	/** Current value of KernelFunctionName as display text. */
+	FText GetCurrentKernelFunctionNameText() const;
+	/** Current value of DataInterfaceName as display text. */
+	FText GetCurrentInterfaceNameText() const;
+	/** Current value of DataInterfaceFunctionName as display text. */
+	FText GetCurrentFunctionNameText() const;
+
+	/** Writes the selected kernel function name back to the property. */
+	void OnKernelFunctionSelected(TSharedPtr<FString> NewValue, ESelectInfo::Type SelectInfo);
+	/** Writes the selected interface name back to the property and refreshes the function list. */
+	void OnInterfaceSelected(TSharedPtr<FName> NewValue, ESelectInfo::Type SelectInfo);
+	/** Writes the selected DI function name back to the property. */
+	void OnFunctionSelected(TSharedPtr<FString> NewValue, ESelectInfo::Type SelectInfo);
+
+	/** True when customizing an input pin. False for an output pin. */
+	bool bIsInput = true;
+
+	/** Handle for the FKernelPin struct itself. */
+	TSharedPtr<IPropertyHandle> StructHandle;
+	/** Handle for FKernelPin::KernelFunctionName. */
+	TSharedPtr<IPropertyHandle> KernelFnHandle;
+	/** Handle for FKernelPin::DataInterfaceName. */
+	TSharedPtr<IPropertyHandle> InterfaceNameHandle;
+	/** Handle for FKernelPin::DataInterfaceFunctionName. */
+	TSharedPtr<IPropertyHandle> FunctionNameHandle;
+	/** Handle for FKernelPin::bOrphaned. */
+	TSharedPtr<IPropertyHandle> OrphanedHandle;
+	/** Handle for the owning FComputeGraphKernelDesc::SourceText. */
+	TSharedPtr<IPropertyHandle> KernelSourceHandle;
+	/** Handle for the owning FComputeGraphKernelDesc::EntryPoint. */
+	TSharedPtr<IPropertyHandle> KernelEntryPointHandle;
+
+	/** Populated by RefreshKernelFunctionNameOptions(). Drives the kernel function combo box. */
+	TArray<TSharedPtr<FString>> KernelFunctionNameOptions;
+	/** Populated by RefreshInterfaceNameOptions(). Drives the data interface combo box. */
+	TArray<TSharedPtr<FName>> InterfaceNameOptions;
+	/** Populated by RefreshFunctionNameOptions(). Drives the DI function combo box. */
+	TArray<TSharedPtr<FString>> FunctionNameOptions;
+};
diff --git a/Engine/Plugins/Runtime/ComputeFramework/Source/EditableComputeGraphEditor/Private/ComputeGraphHlslEditor.cpp b/Engine/Plugins/Runtime/ComputeFramework/Source/EditableComputeGraphEditor/Private/ComputeGraphHlslEditor.cpp
new file mode 100644
--- /dev/null
+++ b/Engine/Plugins/Runtime/ComputeFramework/Source/EditableComputeGraphEditor/Private/ComputeGraphHlslEditor.cpp
@@ -0,0 +1,181 @@
+// Copyright Epic Games, Inc. All Rights Reserved.
+
+#include "ComputeFramework/ComputeGraphHlslEditor.h"
+
+#include "ComputeFramework/ComputeGraphHlslSyntaxHighlighter.h"
+#include "ComputeFramework/EditableComputeGraph.h"
+#include "Fonts/SlateFontInfo.h"
+#include "Styling/AppStyle.h"
+#include "Styling/CoreStyle.h"
+#include "Widgets/Input/SMultiLineEditableTextBox.h"
+#include "Widgets/Layout/SBorder.h"
+#include "Widgets/Layout/SBox.h"
+#include "Widgets/SBoxPanel.h"
+#include "Widgets/Text/STextBlock.h"
+
+#define LOCTEXT_NAMESPACE "ComputeFrameworkEditor"
+
+void SComputeGraphHlslEditor::Construct(FArguments const& InArgs)
+{
+	Asset = InArgs._Asset;
+	OnTextCommitted_WithKernel = InArgs._OnTextCommitted_WithKernel;
+	OnTextChanged = InArgs._OnTextChanged;
+	OnCompileRequested = InArgs._OnCompileRequested;
+
+	FSlateFontInfo MonoFont = FCoreStyle::GetDefaultFontStyle("Mono", 10);
+
+	// Bound DI_ function names are highlighted with a teal colour.
+	FTextBlockStyle BoundFunctionStyle = FAppStyle::Get().GetWidgetStyle<FTextBlockStyle>("SyntaxHighlight.SourceCode.Normal");
+	BoundFunctionStyle.SetColorAndOpacity(FLinearColor(0.2f, 0.85f, 0.75f, 1.f));
+
+	SyntaxHighlighter = FComputeGraphHlslSyntaxHighlighter::Create(
+		FHLSLSyntaxHighlighterMarshaller::FSyntaxTextStyle(
+			FAppStyle::Get().GetWidgetStyle<FTextBlockStyle>("SyntaxHighlight.SourceCode.Normal"),
+			FAppStyle::Get().GetWidgetStyle<FTextBlockStyle>("SyntaxHighlight.SourceCode.Operator"),
+			FAppStyle::Get().GetWidgetStyle<FTextBlockStyle>("SyntaxHighlight.SourceCode.Keyword"),
+			FAppStyle::Get().GetWidgetStyle<FTextBlockStyle>("SyntaxHighlight.SourceCode.String"),
+			FAppStyle::Get().GetWidgetStyle<FTextBlockStyle>("SyntaxHighlight.SourceCode.Number"),
+			FAppStyle::Get().GetWidgetStyle<FTextBlockStyle>("SyntaxHighlight.SourceCode.Comment"),
+			FAppStyle::Get().GetWidgetStyle<FTextBlockStyle>("SyntaxHighlight.SourceCode.PreProcessorKeyword"),
+			FAppStyle::Get().GetWidgetStyle<FTextBlockStyle>("SyntaxHighlight.SourceCode.Error")),
+		BoundFunctionStyle);
+
+	ChildSlot
+	[
+		SNew(SVerticalBox)
+
+		// Header bar showing which kernel is being edited.
+		+ SVerticalBox::Slot()
+		.AutoHeight()
+		[
+			SNew(SBorder)
+			.BorderImage(FAppStyle::Get().GetBrush("ToolPanel.GroupBorder"))
+			.Padding(FMargin(8.f, 4.f))
+			[
+				SAssignNew(PlaceholderText, STextBlock)
+				.Text(LOCTEXT("NoKernelSelected", "Select a kernel from the navigator."))
+				.Font(FAppStyle::Get().GetFontStyle("PropertyWindow.NormalFont"))
+				.ColorAndOpacity(FSlateColor::UseSubduedForeground())
+			]
+		]
+
+		// Text editor area.
+		+ SVerticalBox::Slot()
+		.FillHeight(1.f)
+		[
+			SNew(SBorder)
+			.BorderImage(FAppStyle::Get().GetBrush("NoBorder"))
+			.Padding(0.f)
+			[
+				SAssignNew(TextBox, SMultiLineEditableTextBox)
+				.Font(MonoFont)
+				.Marshaller(SyntaxHighlighter)
+				.AutoWrapText(false)
+				.AllowMultiLine(true)
+				.IsReadOnly(true) // Read-only until a kernel is selected.
+				.OnTextCommitted(this, &SComputeGraphHlslEditor::OnTextCommitted)
+				.OnTextChanged(this, &SComputeGraphHlslEditor::OnTextChangedInternal)
+				.OnKeyDownHandler(this, &SComputeGraphHlslEditor::OnTextBoxKeyDown)
+			]
+		]
+	];
+}
+
+void SComputeGraphHlslEditor::SetActiveKernel(FName KernelName)
+{
+	if (!Asset.IsValid())
+	{
+		return;
+	}
+
+	ActiveKernelName = KernelName;
+
+	// Find the kernel in the description and load its source text.
+	FComputeGraphDesc const& Desc = Asset->GetGraphDescription();
+	for (FComputeGraphKernelDesc const& KernelDesc : Desc.Kernels)
+	{
+		if (KernelDesc.Name != KernelName)
+		{
+			continue;
+		}
+
+		if (TextBox.IsValid())
+		{
+			TextBox->SetText(FText::FromString(KernelDesc.SourceText));
+			TextBox->SetIsReadOnly(false);
+		}
+
+		if (PlaceholderText.IsValid())
+		{
+			PlaceholderText->SetText(FText::Format(LOCTEXT("KernelHeader", "Kernel: {0}"), FText::FromName(KernelName)));
+		}
+		return;
+	}
+
+	// Kernel not found. Most likely it was just renamed/deleted.
+	if (TextBox.IsValid())
+	{
+		TextBox->SetText(FText::GetEmpty());
+		TextBox->SetIsReadOnly(true);
+	}
+	if (PlaceholderText.IsValid())
+	{
+		PlaceholderText->SetText(LOCTEXT("NoKernelSelected", "Select a kernel from the navigator."));
+	}
+	ActiveKernelName = NAME_None;
+}
+
+void SComputeGraphHlslEditor::SetBoundFunctions(TArray<FString> const& Names)
+{
+	if (SyntaxHighlighter.IsValid())
+	{
+		SyntaxHighlighter->SetBoundFunctions(Names);
+	}
+}
+
+void SComputeGraphHlslEditor::OnTextCommitted(FText const& NewText, ETextCommit::Type CommitType)
+{
+	if (!ActiveKernelName.IsNone() && CommitType != ETextCommit::OnCleared)
+	{
+		OnTextCommitted_WithKernel.ExecuteIfBound(ActiveKernelName, NewText.ToString());
+	}
+}
+
+void SComputeGraphHlslEditor::CommitCurrentText()
+{
+	if (!ActiveKernelName.IsNone() && TextBox.IsValid())
+	{
+		OnTextCommitted_WithKernel.ExecuteIfBound(ActiveKernelName, TextBox->GetText().ToString());
+	}
+}
+
+void SComputeGraphHlslEditor::OnTextChangedInternal(FText const& /*NewText*/)
+{
+	OnTextChanged.ExecuteIfBound();
+}
+
+FReply SComputeGraphHlslEditor::OnKeyDown(FGeometry const& MyGeometry, FKeyEvent const& InKeyEvent)
+{
+	// F8 triggers a compile.
+	if (InKeyEvent.GetKey() == EKeys::F8)
+	{
+		OnCompileRequested.ExecuteIfBound();
+		return FReply::Handled();
+	}
+	return SCompoundWidget::OnKeyDown(MyGeometry, InKeyEvent);
+}
+
+FReply SComputeGraphHlslEditor::OnTextBoxKeyDown(FGeometry const& MyGeometry, FKeyEvent const& InKeyEvent)
+{
+	if (InKeyEvent.GetKey() == EKeys::Tab && TextBox.IsValid())
+	{
+		// Insert spaces to the next 4-column tab stop so Tab doesn't navigate focus away.
+		const int32 Column = TextBox->GetCursorLocation().GetOffset();
+		const int32 NumSpaces = 4 - (Column % 4);
+		TextBox->InsertTextAtCursor(FString::ChrN(NumSpaces, TEXT(' ')));
+		return FReply::Handled();
+	}
+	return FReply::Unhandled();
+}
+
+#undef LOCTEXT_NAMESPACE
diff --git a/Engine/Plugins/Runtime/ComputeFramework/Source/EditableComputeGraphEditor/Private/ComputeGraphHlslSyntaxHighlighter.cpp b/Engine/Plugins/Runtime/ComputeFramework/Source/EditableComputeGraphEditor/Private/ComputeGraphHlslSyntaxHighlighter.cpp
new file mode 100644
--- /dev/null
+++ b/Engine/Plugins/Runtime/ComputeFramework/Source/EditableComputeGraphEditor/Private/ComputeGraphHlslSyntaxHighlighter.cpp
@@ -0,0 +1,55 @@
+// Copyright Epic Games, Inc. All Rights Reserved.
+
+#include "ComputeFramework/ComputeGraphHlslSyntaxHighlighter.h"
+
+#include "Framework/Text/IRun.h"
+#include "Framework/Text/SlateTextRun.h"
+#include "Framework/Text/TextLayout.h"
+#include "Text/HLSLSyntaxTokenizer.h"
+
+TSharedRef<FComputeGraphHlslSyntaxHighlighter> FComputeGraphHlslSyntaxHighlighter::Create(FSyntaxTextStyle const& InSyntaxTextStyle, FTextBlockStyle const& InBoundFunctionStyle)
+{
+	return MakeShareable(new FComputeGraphHlslSyntaxHighlighter(CreateTokenizer(), InSyntaxTextStyle, InBoundFunctionStyle));
+}
+
+FComputeGraphHlslSyntaxHighlighter::FComputeGraphHlslSyntaxHighlighter(
+	TSharedPtr<ISyntaxTokenizer> InTokenizer,
+	FSyntaxTextStyle const& InSyntaxTextStyle,
+	FTextBlockStyle const& InBoundFunctionStyle)
+	: FHLSLSyntaxHighlighterMarshaller(MoveTemp(InTokenizer), InSyntaxTextStyle)
+	, BoundFunctionStyle(InBoundFunctionStyle)
+{
+}
+
+void FComputeGraphHlslSyntaxHighlighter::SetBoundFunctions(TArray<FString> const& Names)
+{
+	BoundFunctionNames = TSet<FString>(Names);
+	MakeDirty();
+}
+
+FTextLayout::FNewLineData FComputeGraphHlslSyntaxHighlighter::ProcessTokenizedLine(
+	ISyntaxTokenizer::FTokenizedLine const& TokenizedLine,
+	int32 const& LineNumber,
+	FString const& SourceString,
+	EParseState& ParseState)
+{
+	FTextLayout::FNewLineData LineData = FHLSLSyntaxHighlighterMarshaller::ProcessTokenizedLine(TokenizedLine, LineNumber, SourceString, ParseState);
+
+	if (BoundFunctionNames.IsEmpty())
+	{
+		return LineData;
+	}
+
+	// Replace runs whose text matches a bound function name.
+	for (TSharedRef<IRun>& Run : LineData.Runs)
+	{
+		const FTextRange Range = Run->GetTextRange();
+		const FString Token = LineData.Text->Mid(Range.BeginIndex, Range.Len());
+		if (BoundFunctionNames.Contains(Token))
+		{
+			Run = FSlateTextRun::Create(FRunInfo(TEXT("SyntaxHighlight.HLSL.BoundFunction")), LineData.Text, BoundFunctionStyle, Range);
+		}
+	}
+
+	return LineData;
+}
diff --git a/Engine/Plugins/Runtime/ComputeFramework/Source/EditableComputeGraphEditor/Private/ComputeGraphNavigator.cpp b/Engine/Plugins/Runtime/ComputeFramework/Source/EditableComputeGraphEditor/Private/ComputeGraphNavigator.cpp
new file mode 100644
--- /dev/null
+++ b/Engine/Plugins/Runtime/ComputeFramework/Source/EditableComputeGraphEditor/Private/ComputeGraphNavigator.cpp
@@ -0,0 +1,512 @@
+// Copyright Epic Games, Inc. All Rights Reserved.
+
+#include "ComputeFramework/ComputeGraphNavigator.h"
+
+#include "ComputeFramework/EditableComputeGraph.h"
+#include "Framework/Commands/UIAction.h"
+#include "Input/Events.h"
+#include "InputCoreTypes.h"
+#include "Framework/MultiBox/MultiBoxBuilder.h"
+#include "Styling/AppStyle.h"
+#include "Styling/StarshipCoreStyle.h"
+#include "Widgets/Input/SButton.h"
+#include "Widgets/Images/SImage.h"
+#include "Widgets/Layout/SBorder.h"
+#include "Widgets/SBoxPanel.h"
+#include "Widgets/Text/SInlineEditableTextBlock.h"
+#include "Widgets/Text/STextBlock.h"
+#include "Widgets/Views/SExpanderArrow.h"
+#include "Widgets/Views/STableRow.h"
+
+#define LOCTEXT_NAMESPACE "ComputeFrameworkEditor"
+
+/** Internal tree node for SComputeGraphNavigator. Either a section header or a leaf item. */
+struct FNavigatorItem
+{
+	/** Display name shown in the row. */
+	FName Name;
+	/** Which kind of graph item this node represents. */
+	EComputeGraphItemKind Kind = EComputeGraphItemKind::None;
+	/** Index into the relevant array in FComputeGraphDesc. INDEX_NONE for header nodes. */
+	int32 Index = INDEX_NONE;
+	/** True for section header nodes, false for leaf item nodes. */
+	bool bHeader = false;
+	/** Child leaf items. Populated for header nodes only. */
+	TArray<FNavigatorItemPtr> Children;
+
+	static TSharedRef<FNavigatorItem> MakeHeader(FName InName, EComputeGraphItemKind InKind)
+	{
+		TSharedRef<FNavigatorItem> Item = MakeShared<FNavigatorItem>();
+		Item->Name = InName;
+		Item->Kind = InKind;
+		Item->bHeader = true;
+		return Item;
+	}
+
+	static TSharedRef<FNavigatorItem> MakeLeaf(FName InName, EComputeGraphItemKind InKind, int32 InIndex)
+	{
+		TSharedRef<FNavigatorItem> Item = MakeShared<FNavigatorItem>();
+		Item->Name = InName;
+		Item->Kind = InKind;
+		Item->Index = InIndex;
+		return Item;
+	}
+};
+
+/** Tree row for a section header. Displays the section name, an expander arrow, and an Add button. */
+class SNavigatorHeaderRow : public STableRow<FNavigatorItemPtr>
+{
+public:
+	SLATE_BEGIN_ARGS(SNavigatorHeaderRow) {}
+		SLATE_ARGUMENT(FNavigatorItemPtr, Item)
+		SLATE_EVENT(FOnComputeGraphAddItem, OnAddItem)
+	SLATE_END_ARGS()
+
+	void Construct(FArguments const& InArgs, TSharedRef<STableViewBase> const& InOwnerTable)
+	{
+		ItemPtr = InArgs._Item;
+		OnAddItem = InArgs._OnAddItem;
+
+		STableRow<FNavigatorItemPtr>::ConstructInternal(
+			STableRow<FNavigatorItemPtr>::FArguments()
+			.Style(FAppStyle::Get(), "DetailsView.TreeView.TableRow")
+			.ShowSelection(false),
+			InOwnerTable);
+
+		ChildSlot
+		.Padding(0.f, 2.f, 0.f, 0.f)
+		[
+			SNew(SBorder)
+			.BorderImage(this, &SNavigatorHeaderRow::GetBorderBrush)
+			.Padding(FMargin(3.f, 5.f))
+			[
+				SNew(SHorizontalBox)
+
+				// Expander arrow
+				+ SHorizontalBox::Slot()
+				.AutoWidth()
+				.VAlign(VAlign_Center)
+				[
+					SNew(SExpanderArrow, SharedThis(this))
+					.IndentAmount(8)
+					.ShouldDrawWires(false)
+				]
+
+				// Section title
+				+ SHorizontalBox::Slot()
+				.FillWidth(1.f)
+				.VAlign(VAlign_Center)
+				.Padding(2.f, 0.f)
+				[
+					SNew(STextBlock)
+					.Text(FText::FromName(ItemPtr->Name))
+					.TextStyle(FAppStyle::Get(), "DetailsView.CategoryTextStyle")
+					.Font(FStyleFonts::Get().NormalBold)
+				]
+
+				// Add button
+				+ SHorizontalBox::Slot()
+				.AutoWidth()
+				.VAlign(VAlign_Center)
+				[
+					SNew(SButton)
+					.ButtonStyle(FAppStyle::Get(), "SimpleButton")
+					.ContentPadding(FMargin(4.f, 0.f))
+					.ToolTipText(LOCTEXT("AddItemTooltip", "Add new item"))
+					.OnClicked_Lambda([this]()
+					{
+						OnAddItem.ExecuteIfBound(ItemPtr->Kind);
+						return FReply::Handled();
+					})
+					[
+						SNew(SImage)
+						.Image(FAppStyle::Get().GetBrush("Icons.PlusCircle"))
+						.ColorAndOpacity(FSlateColor::UseForeground())
+					]
+				]
+			]
+		];
+	}
+
+	/** Clicking anywhere on the header bar toggles expansion. */
+	FReply OnMouseButtonDown(FGeometry const& MyGeometry, FPointerEvent const& MouseEvent) override
+	{
+		if (MouseEvent.GetEffectingButton() == EKeys::LeftMouseButton)
+		{
+			ToggleExpansion();
+			return FReply::Handled();
+		}
+		return STableRow<FNavigatorItemPtr>::OnMouseButtonDown(MyGeometry, MouseEvent);
+	}
+
+private:
+	FSlateBrush const* GetBorderBrush() const
+	{
+		return IsHovered() ? FAppStyle::Get().GetBrush("Brushes.Secondary") : FAppStyle::Get().GetBrush("Brushes.Header");
+	}
+
+	/** The header item this row represents. */
+	FNavigatorItemPtr ItemPtr;
+	/** Delegate fired when the user clicks the Add button. */
+	FOnComputeGraphAddItem OnAddItem;
+};
+
+/** Tree row for a leaf item. Displays the item name with inline rename support. */
+class SNavigatorLeafRow : public STableRow<FNavigatorItemPtr>
+{
+public:
+	SLATE_BEGIN_ARGS(SNavigatorLeafRow) {}
+		SLATE_ARGUMENT(FNavigatorItemPtr, Item)
+		SLATE_EVENT(FOnComputeGraphRenameItem, OnRenameItem)
+		SLATE_EVENT(FOnVerifyTextChanged, OnVerifyRenameText)
+	SLATE_END_ARGS()
+
+	void Construct(FArguments const& InArgs, TSharedRef<STableViewBase> const& InOwnerTable)
+	{
+		ItemPtr = InArgs._Item;
+		OnRenameItem = InArgs._OnRenameItem;
+
+		STableRow<FNavigatorItemPtr>::Construct(
+			STableRow<FNavigatorItemPtr>::FArguments()
+			.Padding(FMargin(20.f, 1.f, 4.f, 1.f))
+			.Content()
+			[
+				SAssignNew(InlineText, SInlineEditableTextBlock)
+				.Text(FText::FromName(ItemPtr->Name))
+				.Font(FAppStyle::Get().GetFontStyle("NormalFont"))
+				.IsSelected(this, &SNavigatorLeafRow::IsSelectedExclusively)
+				.OnVerifyTextChanged(InArgs._OnVerifyRenameText)
+				.OnTextCommitted(this, &SNavigatorLeafRow::OnTextCommitted)
+			],
+			InOwnerTable);
+	}
+
+	void EnterRenameMode()
+	{
+		if (InlineText.IsValid())
+		{
+			InlineText->EnterEditingMode();
+		}
+	}
+
+private:
+	void OnTextCommitted(FText const& NewText, ETextCommit::Type CommitType)
+	{
+		const FString Trimmed = NewText.ToString().TrimStartAndEnd();
+		if (!Trimmed.IsEmpty() && CommitType != ETextCommit::OnCleared)
+		{
+			OnRenameItem.ExecuteIfBound(ItemPtr->Kind, ItemPtr->Index, FName(*Trimmed));
+		}
+	}
+
+	/** The leaf item this row represents. */
+	FNavigatorItemPtr ItemPtr;
+	/** The editable text block used for inline renaming. */
+	TSharedPtr<SInlineEditableTextBlock> InlineText;
+	/** Delegate fired when the user commits a rename. */
+	FOnComputeGraphRenameItem OnRenameItem;
+};
+
+void SComputeGraphNavigator::Construct(FArguments const& InArgs)
+{
+	Asset = InArgs._Asset;
+	OnItemSelected = InArgs._OnItemSelected;
+	OnAddItem = InArgs._OnAddItem;
+	OnDeleteItem = InArgs._OnDeleteItem;
+	OnDuplicateItem = InArgs._OnDuplicateItem;
+	OnRenameItem = InArgs._OnRenameItem;
+	OnKernelSelected = InArgs._OnKernelSelected;
+
+	RebuildItems();
+
+	ChildSlot
+	[
+		SNew(SBorder)
+		.BorderImage(FAppStyle::Get().GetBrush("ToolPanel.GroupBorder"))
+		.Padding(0.f)
+		[
+			SAssignNew(TreeView, STreeView<FNavigatorItemPtr>)
+			.TreeItemsSource(&RootItems)
+			.OnGenerateRow(this, &SComputeGraphNavigator::GenerateRow)
+			.OnGetChildren(this, &SComputeGraphNavigator::OnGetChildren)
+			.OnSelectionChanged(this, &SComputeGraphNavigator::OnSelectionChanged)
+			.OnExpansionChanged(this, &SComputeGraphNavigator::OnExpansionChanged)
+			.OnContextMenuOpening(this, &SComputeGraphNavigator::BuildContextMenu)
+			.SelectionMode(ESelectionMode::Single)
+		]
+	];
+
+	// Expand all sections by default.
+	for (FNavigatorItemPtr const& Root : RootItems)
+	{
+		TreeView->SetItemExpansion(Root, true);
+	}
+}
+
+void SComputeGraphNavigator::Refresh()
+{
+	RebuildItems();
+	if (!TreeView.IsValid()) return;
+
+	TreeView->RequestTreeRefresh();
+
+	// Restore expansion state.
+	// Expand unless the section was manually collapsed.
+	for (FNavigatorItemPtr const& Root : RootItems)
+	{
+		TreeView->SetItemExpansion(Root, !CollapsedSections.Contains(Root->Kind));
+	}
+
+	// Re-apply selection by matching kind + index in the new leaf items.
+	if (Selection.Kind != EComputeGraphItemKind::None && Selection.Index != INDEX_NONE)
+	{
+		for (FNavigatorItemPtr const& Root : RootItems)
+		{
+			for (FNavigatorItemPtr const& Child : Root->Children)
+			{
+				if (Child->Kind == Selection.Kind && Child->Index == Selection.Index)
+				{
+					TreeView->SetSelection(Child, ESelectInfo::Direct);
+					return;
+				}
+			}
+		}
+	}
+}
+
+void SComputeGraphNavigator::SetSelectedItem(EComputeGraphItemKind Kind, int32 Index)
+{
+	Selection.Kind = Kind;
+	Selection.Index = Index;
+
+	if (!TreeView.IsValid()) return;
+
+	if (Kind == EComputeGraphItemKind::None || Index == INDEX_NONE)
+	{
+		TreeView->ClearSelection();
+		return;
+	}
+
+	for (FNavigatorItemPtr const& Root : RootItems)
+	{
+		for (FNavigatorItemPtr const& Child : Root->Children)
+		{
+			if (Child->Kind == Kind && Child->Index == Index)
+			{
+				TreeView->SetSelection(Child, ESelectInfo::Direct);
+				return;
+			}
+		}
+	}
+}
+
+void SComputeGraphNavigator::RebuildItems()
+{
+	RootItems.Reset();
+	if (!Asset.IsValid()) return;
+
+	FComputeGraphDesc const& Desc = Asset->GetGraphDescription();
+
+	// Binding Objects.
+	TSharedRef<FNavigatorItem> BindingSection = FNavigatorItem::MakeHeader(TEXT("Binding Objects"), EComputeGraphItemKind::BindingObject);
+	for (int32 Index = 0; Index < Desc.BindingObjects.Num(); ++Index)
+	{
+		const FName N = Desc.BindingObjects[Index].Name;
+		BindingSection->Children.Add(FNavigatorItem::MakeLeaf(N.IsNone() ? FName(TEXT("(unnamed)")) : N, EComputeGraphItemKind::BindingObject, Index));
+	}
+	RootItems.Add(BindingSection);
+
+	// Data Interfaces.
+	TSharedRef<FNavigatorItem> InterfaceSection = FNavigatorItem::MakeHeader(TEXT("Data Interfaces"), EComputeGraphItemKind::Interface);
+	for (int32 Index = 0; Index < Desc.DataInterfaces.Num(); ++Index)
+	{
+		const FName N = Desc.DataInterfaces[Index].Name;
+		InterfaceSection->Children.Add(FNavigatorItem::MakeLeaf(N.IsNone() ? FName(TEXT("(unnamed)")) : N, EComputeGraphItemKind::Interface, Index));
+	}
+	RootItems.Add(InterfaceSection);
+
+	// Kernels.
+	TSharedRef<FNavigatorItem> KernelSection = FNavigatorItem::MakeHeader(TEXT("Kernels"), EComputeGraphItemKind::Kernel);
+	for (int32 Index = 0; Index < Desc.Kernels.Num(); ++Index)
+	{
+		const FName N = Desc.Kernels[Index].Name;
+		KernelSection->Children.Add(FNavigatorItem::MakeLeaf(N.IsNone() ? FName(TEXT("(unnamed)")) : N, EComputeGraphItemKind::Kernel, Index));
+	}
+	RootItems.Add(KernelSection);
+}
+
+TSharedRef<ITableRow> SComputeGraphNavigator::GenerateRow(
+	FNavigatorItemPtr Item,
+	TSharedRef<STableViewBase> const& OwnerTable)
+{
+	if (Item->bHeader)
+	{
+		return SNew(SNavigatorHeaderRow, OwnerTable)
+			.Item(Item)
+			.OnAddItem(FOnComputeGraphAddItem::CreateSP(this, &SComputeGraphNavigator::OnAddItemClicked));
+	}
+
+	const EComputeGraphItemKind Kind = Item->Kind;
+	const int32 Index = Item->Index;
+	return SNew(SNavigatorLeafRow, OwnerTable)
+		.Item(Item)
+		.OnRenameItem(FOnComputeGraphRenameItem::CreateSP(this, &SComputeGraphNavigator::OnRenameItemCommitted))
+		.OnVerifyRenameText(FOnVerifyTextChanged::CreateSP(this, &SComputeGraphNavigator::VerifyRenameText, Kind, Index));
+}
+
+void SComputeGraphNavigator::OnGetChildren(FNavigatorItemPtr Item, TArray<FNavigatorItemPtr>& OutChildren)
+{
+	OutChildren.Append(Item->Children);
+}
+
+TSharedPtr<SWidget> SComputeGraphNavigator::BuildContextMenu()
+{
+	TArray<FNavigatorItemPtr> const SelectedItems = TreeView->GetSelectedItems();
+	if (SelectedItems.IsEmpty() || SelectedItems[0]->bHeader)
+	{
+		return nullptr;
+	}
+
+	FMenuBuilder MenuBuilder(/*bInShouldCloseWindowAfterMenuSelection=*/true, nullptr);
+
+	MenuBuilder.AddMenuEntry(
+		LOCTEXT("DeleteItem", "Delete"),
+		LOCTEXT("DeleteItemTooltip", "Remove this item from the graph description."),
+		FSlateIcon(FAppStyle::GetAppStyleSetName(), "Icons.Delete"),
+		FUIAction(FExecuteAction::CreateSP(this, &SComputeGraphNavigator::OnDeleteSelected)));
+
+	MenuBuilder.AddMenuEntry(
+		LOCTEXT("DuplicateItem", "Duplicate"),
+		LOCTEXT("DuplicateItemTooltip", "Add a copy of this item with a new name."),
+		FSlateIcon(FAppStyle::GetAppStyleSetName(), "Icons.Duplicate"),
+		FUIAction(FExecuteAction::CreateSP(this, &SComputeGraphNavigator::OnDuplicateSelected)));
+
+	MenuBuilder.AddMenuEntry(
+		LOCTEXT("RenameItem", "Rename"),
+		LOCTEXT("RenameItemTooltip", "Rename this item inline."),
+		FSlateIcon(FAppStyle::GetAppStyleSetName(), "GenericCommands.Rename"),
+		FUIAction(FExecuteAction::CreateSP(this, &SComputeGraphNavigator::OnRenameSelected)));
+
+	return MenuBuilder.MakeWidget();
+}
+
+void SComputeGraphNavigator::OnSelectionChanged(FNavigatorItemPtr Item, ESelectInfo::Type SelectInfo)
+{
+	if (!Item.IsValid() || Item->bHeader)
+	{
+		Selection.Kind = EComputeGraphItemKind::None;
+		Selection.Index = INDEX_NONE;
+		OnItemSelected.ExecuteIfBound(EComputeGraphItemKind::None, INDEX_NONE, NAME_None);
+		return;
+	}
+
+	Selection.Kind = Item->Kind;
+	Selection.Index = Item->Index;
+	OnItemSelected.ExecuteIfBound(Item->Kind, Item->Index, Item->Name);
+
+	if (Item->Kind == EComputeGraphItemKind::Kernel)
+	{
+		OnKernelSelected.ExecuteIfBound(Item->Name);
+	}
+}
+
+void SComputeGraphNavigator::OnExpansionChanged(FNavigatorItemPtr Item, bool bExpanded)
+{
+	if (!Item.IsValid() || !Item->bHeader) return;
+
+	if (bExpanded)
+	{
+		CollapsedSections.Remove(Item->Kind);
+	}
+	else
+	{
+		CollapsedSections.Add(Item->Kind);
+	}
+}
+
+void SComputeGraphNavigator::OnAddItemClicked(EComputeGraphItemKind Kind)
+{
+	OnAddItem.ExecuteIfBound(Kind);
+}
+
+void SComputeGraphNavigator::OnDeleteSelected()
+{
+	if (Selection.Kind != EComputeGraphItemKind::None && Selection.Index != INDEX_NONE)
+	{
+		OnDeleteItem.ExecuteIfBound(Selection.Kind, Selection.Index);
+	}
+}
+
+void SComputeGraphNavigator::OnDuplicateSelected()
+{
+	if (Selection.Kind != EComputeGraphItemKind::None && Selection.Index != INDEX_NONE)
+	{
+		OnDuplicateItem.ExecuteIfBound(Selection.Kind, Selection.Index);
+	}
+}
+
+void SComputeGraphNavigator::OnRenameItemCommitted(EComputeGraphItemKind Kind, int32 Index, FName NewName)
+{
+	OnRenameItem.ExecuteIfBound(Kind, Index, NewName);
+}
+
+bool SComputeGraphNavigator::VerifyRenameText(FText const& NewText, FText& OutError, EComputeGraphItemKind Kind, int32 Index) const
+{
+	const FString Trimmed = NewText.ToString().TrimStartAndEnd();
+	if (Trimmed.IsEmpty())
+	{
+		OutError = LOCTEXT("RenameEmpty", "Name cannot be empty.");
+		return false;
+	}
+
+	const FName Candidate(*Trimmed);
+
+	for (FNavigatorItemPtr const& Root : RootItems)
+	{
+		for (FNavigatorItemPtr const& Child : Root->Children)
+		{
+			if (Child->Kind == Kind && Child->Index != Index && Child->Name == Candidate)
+			{
+				OutError = FText::Format(LOCTEXT("RenameDuplicate", "'{0}' is already in use."), NewText);
+				return false;
+			}
+		}
+	}
+	return true;
+}
+
+FReply SComputeGraphNavigator::OnKeyDown(FGeometry const& MyGeometry, FKeyEvent const& InKeyEvent)
+{
+	if (InKeyEvent.GetKey() == EKeys::F2)
+	{
+		OnRenameSelected();
+		return FReply::Handled();
+	}
+	return SCompoundWidget::OnKeyDown(MyGeometry, InKeyEvent);
+}
+
+void SComputeGraphNavigator::OnRenameSelected()
+{
+	if (Selection.Kind == EComputeGraphItemKind::None || Selection.Index == INDEX_NONE) 
+	{
+		return;
+	}
+
+	for (FNavigatorItemPtr const& Root : RootItems)
+	{
+		for (FNavigatorItemPtr const& Child : Root->Children)
+		{
+			if (Child->Kind == Selection.Kind && Child->Index == Selection.Index)
+			{
+				TSharedPtr<ITableRow> Row = TreeView->WidgetFromItem(Child);
+				if (TSharedPtr<SNavigatorLeafRow> LeafRow = StaticCastSharedPtr<SNavigatorLeafRow>(Row))
+				{
+					LeafRow->EnterRenameMode();
+				}
+				return;
+			}
+		}
+	}
+}
+
+#undef LOCTEXT_NAMESPACE
diff --git a/Engine/Plugins/Runtime/ComputeFramework/Source/EditableComputeGraphEditor/Private/ComputeGraphView.cpp b/Engine/Plugins/Runtime/ComputeFramework/Source/EditableComputeGraphEditor/Private/ComputeGraphView.cpp
new file mode 100644
--- /dev/null
+++ b/Engine/Plugins/Runtime/ComputeFramework/Source/EditableComputeGraphEditor/Private/ComputeGraphView.cpp
@@ -0,0 +1,458 @@
+// Copyright Epic Games, Inc. All Rights Reserved.
+
+#include "ComputeFramework/ComputeGraphView.h"
+
+#include "ComputeFramework/EditableComputeGraph.h"
+#include "Input/Events.h"
+#include "Rendering/DrawElements.h"
+#include "Styling/AppStyle.h"
+#include "Styling/SlateBrush.h"
+#include "Widgets/Layout/SBorder.h"
+
+#define LOCTEXT_NAMESPACE "ComputeFrameworkEditor"
+
+namespace
+{
+	// Layout constants (local pixels, unscaled by DPI).
+	constexpr float NodeWidth = 160.f;
+	constexpr float NodeHeight = 36.f;
+	constexpr float NodeVGap = 12.f;
+	constexpr float ColXStart = 20.f;
+	constexpr float ColStep = NodeWidth + 60.f; // 220 px per rank column
+	constexpr float TopPadding = 20.f;
+
+	constexpr float ViewMinWidth = ColXStart + ColStep + 20.f;
+	constexpr float ViewMinHeight = 200.f;
+
+	// Colours.
+	static const FLinearColor ColNode(0.15f, 0.15f, 0.18f, 1.f);
+	static const FLinearColor ColNodeBorder(0.40f, 0.40f, 0.45f, 1.f);
+	static const FLinearColor ColNodeBorderSelected(1.00f, 0.70f, 0.10f, 1.f);
+	static const FLinearColor ColNodeHighlightKernel(0.20f, 0.35f, 0.45f, 1.f);
+	static const FLinearColor ColNodeHighlightInterface(0.35f, 0.20f, 0.45f, 1.f);
+	static const FLinearColor ColNodeBorderHighlightKernel(0.50f, 0.80f, 1.00f, 1.f);
+	static const FLinearColor ColNodeBorderHighlightInterface(0.85f, 0.50f, 1.00f, 1.f);
+	static const FLinearColor ColEdge(0.55f, 0.75f, 0.55f, 0.8f);
+	static const FLinearColor ColEdgeOrphaned(0.90f, 0.20f, 0.20f, 0.8f);
+	static const FLinearColor ColText(0.90f, 0.90f, 0.90f, 1.f);
+
+} // namespace
+
+void SComputeGraphView::Construct(FArguments const& InArgs)
+{
+	Asset = InArgs._Asset;
+	OnNodeClicked = InArgs._OnNodeClicked;
+	RebuildLayout();
+}
+
+void SComputeGraphView::SetHighlightedKernel(FName KernelName)
+{
+	HighlightedKernel = KernelName;
+	Invalidate(EInvalidateWidgetReason::Paint);
+}
+
+void SComputeGraphView::SetHighlightedInterface(FName InterfaceName)
+{
+	HighlightedInterface = InterfaceName;
+	Invalidate(EInvalidateWidgetReason::Paint);
+}
+
+void SComputeGraphView::Refresh()
+{
+	RebuildLayout();
+	Invalidate(EInvalidateWidgetReason::Layout | EInvalidateWidgetReason::Paint);
+}
+
+void SComputeGraphView::RebuildLayout()
+{
+	NodeLayouts.Reset();
+	Edges.Reset();
+	InterfaceCentreMap.Reset();
+	KernelCentreMap.Reset();
+	MaxRank = 0;
+
+	if (!Asset.IsValid())
+	{
+		return;
+	}
+
+	FComputeGraphDesc const& Desc = Asset->GetGraphDescription();
+
+	// Step 1: Rank assignment via longest-path layering
+	// Rank = length of the longest path from any source node.
+	// Sources are interfaces with no incoming kernel-output edges (rank 0).
+
+	TMap<FName, int32> InterfaceRank;
+	TMap<FName, int32> KernelRank;
+
+	for (FComputeGraphDataInterfaceDesc const& DataInterfaceDesc : Desc.DataInterfaces)
+	{
+		if (!DataInterfaceDesc.Name.IsNone())
+		{
+			InterfaceRank.Add(DataInterfaceDesc.Name, 0);
+		}
+	}
+	for (FComputeGraphKernelDesc const& KernelDesc : Desc.Kernels)
+	{
+		if (!KernelDesc.Name.IsNone())
+		{
+			KernelRank.Add(KernelDesc.Name, 0);
+		}
+	}
+
+	// Which kernels write to each interface?
+	TMap<FName, TArray<FName>> InterfaceWriters;
+	for (FComputeGraphKernelDesc const& KernelDesc : Desc.Kernels)
+	{
+		for (FKernelPin const& Pin : KernelDesc.Outputs)
+		{
+			if (!Pin.DataInterfaceName.IsNone())
+			{
+				InterfaceWriters.FindOrAdd(Pin.DataInterfaceName).AddUnique(KernelDesc.Name);
+			}
+		}
+	}
+
+	// Propagate ranks until stable. Converges for DAGs in at most N+M iterations.
+	// Cap iterations to prevent infinite loops on cyclic graphs.
+	const int32 MaxIterations = Desc.DataInterfaces.Num() + Desc.Kernels.Num() + 1;
+	bool bChanged = true;
+	for (int32 Iteration = 0; bChanged && Iteration < MaxIterations; ++Iteration)
+	{
+		bChanged = false;
+
+		// Kernel rank = max input interface rank + 1.
+		for (FComputeGraphKernelDesc const& KernelDesc : Desc.Kernels)
+		{
+			if (KernelDesc.Name.IsNone())
+			{
+				continue;
+			}
+			int32 MaxInputRank = -1;
+			for (FKernelPin const& Pin : KernelDesc.Inputs)
+			{
+				if (int32* InterfaceRankPtr = InterfaceRank.Find(Pin.DataInterfaceName))
+				{
+					MaxInputRank = FMath::Max(MaxInputRank, *InterfaceRankPtr);
+				}
+			}
+			int32& CachedKernelRank = KernelRank.FindChecked(KernelDesc.Name);
+			const int32 NewKernelRank = MaxInputRank + 1;
+			if (NewKernelRank > CachedKernelRank)
+			{
+				CachedKernelRank = NewKernelRank;
+				bChanged = true;
+			}
+		}
+
+		// Interface rank = max writer-kernel rank + 1 (unchanged if no writers).
+		for (FComputeGraphDataInterfaceDesc const& DataInterfaceDesc : Desc.DataInterfaces)
+		{
+			if (DataInterfaceDesc.Name.IsNone())
+			{
+				continue;
+			}
+			TArray<FName> const* Writers = InterfaceWriters.Find(DataInterfaceDesc.Name);
+			if (!Writers)
+			{
+				continue;
+			}
+			int32 MaxWriterKernelRank = 0;
+			for (FName const& WriterKernelName : *Writers)
+			{
+				if (int32* KernelRankPtr = KernelRank.Find(WriterKernelName))
+				{
+					MaxWriterKernelRank = FMath::Max(MaxWriterKernelRank, *KernelRankPtr);
+				}
+			}
+			int32& CachedInterfaceRank = InterfaceRank.FindChecked(DataInterfaceDesc.Name);
+			const int32 NewInterfaceRank = MaxWriterKernelRank + 1;
+			if (NewInterfaceRank > CachedInterfaceRank)
+			{
+				CachedInterfaceRank = NewInterfaceRank;
+				bChanged = true;
+			}
+		}
+	}
+
+	for (TPair<FName, int32> const& RankPair : InterfaceRank)
+	{
+		MaxRank = FMath::Max(MaxRank, RankPair.Value);
+	}
+	for (TPair<FName, int32> const& RankPair : KernelRank)
+	{
+		MaxRank = FMath::Max(MaxRank, RankPair.Value);
+	}
+
+	// Step 2: Place nodes, stacking vertically within each rank column.
+
+	TMap<int32, int32> NodesPerRank;
+
+	auto PlaceNode = [&](FName Name, bool bKernel, int32 Rank)
+	{
+		int32& Count = NodesPerRank.FindOrAdd(Rank);
+		FNodeLayout& NewLayout = NodeLayouts.AddDefaulted_GetRef();
+		NewLayout.Name = Name;
+		NewLayout.Centre = FVector2D(
+			ColXStart + Rank * ColStep + NodeWidth * 0.5f,
+			TopPadding + Count * (NodeHeight + NodeVGap) + NodeHeight * 0.5f);
+		NewLayout.bKernel = bKernel;
+		NewLayout.Rank = Rank;
+		++Count;
+	};
+
+	// Maintain declaration order within each rank.
+	for (FComputeGraphDataInterfaceDesc const& DataInterfaceDesc : Desc.DataInterfaces)
+	{
+		if (!DataInterfaceDesc.Name.IsNone())
+		{
+			if (int32* RankPtr = InterfaceRank.Find(DataInterfaceDesc.Name))
+			{
+				PlaceNode(DataInterfaceDesc.Name, false, *RankPtr);
+			}
+		}
+	}
+
+	for (FComputeGraphKernelDesc const& KernelDesc : Desc.Kernels)
+	{
+		if (!KernelDesc.Name.IsNone())
+		{
+			if (int32* RankPtr = KernelRank.Find(KernelDesc.Name))
+			{
+				PlaceNode(KernelDesc.Name, true, *RankPtr);
+			}
+		}
+	}
+
+	// Step 3: Build edges.
+
+	auto AddEdges = [&](FComputeGraphKernelDesc const& KernelDesc, TArray<FKernelPin> const& Pins, bool bIsOutput)
+	{
+		for (FKernelPin const& Pin : Pins)
+		{
+			FEdge& NewEdge = Edges.AddDefaulted_GetRef();
+			NewEdge.KernelName = KernelDesc.Name;
+			NewEdge.InterfaceName = Pin.DataInterfaceName;
+			NewEdge.KernelFn = Pin.KernelFunctionName;
+			NewEdge.InterfaceFn = Pin.DataInterfaceFunctionName;
+			NewEdge.bIsOutput = bIsOutput;
+			NewEdge.bOrphaned = Pin.bOrphaned || Pin.DataInterfaceName.IsNone() || Pin.DataInterfaceFunctionName.IsEmpty();
+		}
+	};
+	for (FComputeGraphKernelDesc const& KernelDesc : Desc.Kernels)
+	{
+		AddEdges(KernelDesc, KernelDesc.Inputs, false);
+		AddEdges(KernelDesc, KernelDesc.Outputs, true);
+	}
+
+	// Cache centre maps for OnPaint (avoids per-frame heap allocation).
+	for (FNodeLayout const& NodeLayout : NodeLayouts)
+	{
+		(NodeLayout.bKernel ? KernelCentreMap : InterfaceCentreMap).Add(NodeLayout.Name, NodeLayout.Centre);
+	}
+
+	// Cache height metric for ComputeDesiredSize (NodesPerRank is fully populated above).
+	MaxNodesInAnyRank = 0;
+	for (TPair<int32, int32> const& RankPair : NodesPerRank)
+	{
+		MaxNodesInAnyRank = FMath::Max(MaxNodesInAnyRank, RankPair.Value);
+	}
+}
+
+static void DrawNode(
+	FSlateWindowElementList& OutDrawElements,
+	int32 LayerId,
+	FGeometry const& AllottedGeometry,
+	FVector2D Centre,
+	FName Name,
+	FSlateFontInfo const& Font,
+	FLinearColor Fill,
+	FLinearColor Border,
+	float BorderThickness)
+{
+	const FVector2D TopLeft = Centre - FVector2D(NodeWidth * 0.5f, NodeHeight * 0.5f);
+
+	// For thick (highlighted) borders draw an oversized filled rect behind the node.
+	if (BorderThickness > 1.f)
+	{
+		const FVector2D Expand(BorderThickness, BorderThickness);
+		FSlateDrawElement::MakeBox(
+			OutDrawElements,
+			LayerId,
+			AllottedGeometry.ToPaintGeometry(FVector2D(NodeWidth, NodeHeight) + Expand * 2.f, FSlateLayoutTransform(TopLeft - Expand)),
+			FAppStyle::GetBrush("WhiteBrush"),
+			ESlateDrawEffect::None,
+			Border);
+	}
+
+	// Main node rect.
+	FSlateDrawElement::MakeBox(
+		OutDrawElements,
+		LayerId + 1,
+		AllottedGeometry.ToPaintGeometry(FVector2D(NodeWidth, NodeHeight), FSlateLayoutTransform(TopLeft)),
+		FAppStyle::GetBrush("WhiteBrush"),
+		ESlateDrawEffect::None,
+		Fill);
+
+	// Thin border (non-highlighted). 
+	if (BorderThickness <= 1.f)
+	{
+		FSlateDrawElement::MakeBox(
+			OutDrawElements,
+			LayerId + 1,
+			AllottedGeometry.ToPaintGeometry(FVector2D(NodeWidth, NodeHeight), FSlateLayoutTransform(TopLeft)),
+			FAppStyle::GetBrush("Border"),
+			ESlateDrawEffect::None,
+			Border);
+	}
+
+	// Node name text.
+	FSlateDrawElement::MakeText(
+		OutDrawElements,
+		LayerId + 2,
+		AllottedGeometry.ToPaintGeometry(
+			FVector2D(NodeWidth - 8.f, NodeHeight),
+			FSlateLayoutTransform(TopLeft + FVector2D(4.f, (NodeHeight - Font.Size) * 0.5f))),
+		Name.ToString(),
+		Font,
+		ESlateDrawEffect::None,
+		ColText);
+}
+
+int32 SComputeGraphView::OnPaint(
+	FPaintArgs const& Args,
+	FGeometry const& AllottedGeometry,
+	FSlateRect const& MyCullingRect,
+	FSlateWindowElementList& OutDrawElements,
+	int32 LayerId,
+	FWidgetStyle const& InWidgetStyle,
+	bool bParentEnabled) const
+{
+	// Dark background.
+	FSlateDrawElement::MakeBox(
+		OutDrawElements,
+		LayerId,
+		AllottedGeometry.ToPaintGeometry(),
+		FAppStyle::GetBrush("WhiteBrush"),
+		ESlateDrawEffect::None,
+		FLinearColor(0.07f, 0.07f, 0.07f, 1.f));
+
+	int32 CurrentLayer = LayerId + 1;
+
+	const FSlateFontInfo Font = FAppStyle::Get().GetFontStyle("PropertyWindow.NormalFont");
+
+	// Draw edges first (below nodes).
+	for (FEdge const& Edge : Edges)
+	{
+		FVector2D const* KernelCentre = KernelCentreMap.Find(Edge.KernelName);
+		FVector2D const* InterfaceCentre = InterfaceCentreMap.Find(Edge.InterfaceName);
+		if (!KernelCentre || !InterfaceCentre)
+		{
+			continue;
+		}
+
+		FVector2D P0, P3;
+		if (Edge.bIsOutput)
+		{
+			// Output: right side of kernel to left side of (higher-rank) interface.
+			P0 = *KernelCentre + FVector2D(NodeWidth * 0.5f, 0.f);
+			P3 = *InterfaceCentre - FVector2D(NodeWidth * 0.5f, 0.f);
+		}
+		else
+		{
+			// Input: right side of interface to left side of (higher-rank) kernel.
+			P0 = *InterfaceCentre + FVector2D(NodeWidth * 0.5f, 0.f);
+			P3 = *KernelCentre - FVector2D(NodeWidth * 0.5f, 0.f);
+		}
+
+		const float TanX = (P3.X - P0.X) * 0.4f;
+		const FVector2D P1(P0.X + TanX, P0.Y);
+		const FVector2D P2(P3.X - TanX, P3.Y);
+
+		FSlateDrawElement::MakeCubicBezierSpline(
+			OutDrawElements,
+			CurrentLayer,
+			AllottedGeometry.ToPaintGeometry(),
+			P0, P1, P2, P3,
+			2.f,
+			ESlateDrawEffect::None,
+			Edge.bOrphaned ? ColEdgeOrphaned : ColEdge);
+	}
+	++CurrentLayer;
+
+	// Collect nodes adjacent to the primary selection (connected via an edge).
+	TSet<FName> AdjacentInterfaces;
+	TSet<FName> AdjacentKernels;
+
+	if (!HighlightedKernel.IsNone())
+	{
+		for (FEdge const& Edge : Edges)
+		{
+			if (Edge.KernelName == HighlightedKernel && !Edge.InterfaceName.IsNone())
+			{
+				AdjacentInterfaces.Add(Edge.InterfaceName);
+			}
+		}
+	}
+	if (!HighlightedInterface.IsNone())
+	{
+		for (FEdge const& Edge : Edges)
+		{
+			if (Edge.InterfaceName == HighlightedInterface && !Edge.KernelName.IsNone())
+			{
+				AdjacentKernels.Add(Edge.KernelName);
+			}
+		}
+	}
+
+	// Draw nodes.
+	for (FNodeLayout const& NodeLayout : NodeLayouts)
+	{
+		const bool bIsSelected = (NodeLayout.bKernel && NodeLayout.Name == HighlightedKernel) || (!NodeLayout.bKernel && NodeLayout.Name == HighlightedInterface);
+		const bool bIsAdjacent = !bIsSelected && (NodeLayout.bKernel ? AdjacentKernels.Contains(NodeLayout.Name) : AdjacentInterfaces.Contains(NodeLayout.Name));
+
+		const FLinearColor Fill = (bIsSelected || bIsAdjacent)
+		                        ? (NodeLayout.bKernel ? ColNodeHighlightKernel : ColNodeHighlightInterface)
+		                        : ColNode;
+		const FLinearColor Border = bIsSelected
+		                          ? ColNodeBorderSelected
+		                          : bIsAdjacent ? (NodeLayout.bKernel ? ColNodeBorderHighlightKernel : ColNodeBorderHighlightInterface)
+		                          : ColNodeBorder;
+
+		DrawNode(OutDrawElements, CurrentLayer, AllottedGeometry, NodeLayout.Centre, NodeLayout.Name, Font, Fill, Border, bIsSelected ? 3.f : 1.f);
+	}
+	CurrentLayer += 3;
+
+	return CurrentLayer;
+}
+
+FVector2D SComputeGraphView::ComputeDesiredSize(float LayoutScaleMultiplier) const
+{
+	const float Width = FMath::Max(ViewMinWidth, ColXStart + (MaxRank + 1) * ColStep + 20.f);
+	const float Height = FMath::Max(ViewMinHeight, TopPadding + MaxNodesInAnyRank * (NodeHeight + NodeVGap));
+	return FVector2D(Width, Height);
+}
+
+FReply SComputeGraphView::OnMouseButtonDown(FGeometry const& MyGeometry, FPointerEvent const& MouseEvent)
+{
+	if (MouseEvent.GetEffectingButton() != EKeys::LeftMouseButton || !OnNodeClicked.IsBound())
+	{
+		return FReply::Unhandled();
+	}
+
+	const FVector2D LocalPos = MyGeometry.AbsoluteToLocal(MouseEvent.GetScreenSpacePosition());
+
+	for (FNodeLayout const& NodeLayout : NodeLayouts)
+	{
+		if (FMath::Abs(LocalPos.X - NodeLayout.Centre.X) <= NodeWidth * 0.5f &&
+		    FMath::Abs(LocalPos.Y - NodeLayout.Centre.Y) <= NodeHeight * 0.5f)
+		{
+			OnNodeClicked.Execute(NodeLayout.Name, NodeLayout.bKernel);
+			return FReply::Handled();
+		}
+	}
+
+	return FReply::Unhandled();
+}
+
+#undef LOCTEXT_NAMESPACE
diff --git a/Engine/Plugins/Runtime/ComputeFramework/Source/EditableComputeGraphEditor/Private/EditableComputeGraphAssetDefinition.cpp b/Engine/Plugins/Runtime/ComputeFramework/Source/EditableComputeGraphEditor/Private/EditableComputeGraphAssetDefinition.cpp
new file mode 100644
--- /dev/null
+++ b/Engine/Plugins/Runtime/ComputeFramework/Source/EditableComputeGraphEditor/Private/EditableComputeGraphAssetDefinition.cpp
@@ -0,0 +1,36 @@
+// Copyright Epic Games, Inc. All Rights Reserved.
+
+#include "ComputeFramework/EditableComputeGraphAssetDefinition.h"
+
+#include "ComputeFramework/EditableComputeGraph.h"
+#include "ComputeFramework/EditableComputeGraphEditorToolkit.h"
+
+FText UAssetDefinition_EditableComputeGraph::GetAssetDisplayName() const
+{
+	return NSLOCTEXT("ComputeFramework", "EditableComputeGraphName", "Compute Graph (Experimental)");
+}
+
+FLinearColor UAssetDefinition_EditableComputeGraph::GetAssetColor() const
+{
+	return FColor::Orange;
+}
+
+TSoftClassPtr<UObject> UAssetDefinition_EditableComputeGraph::GetAssetClass() const
+{
+	return UEditableComputeGraph::StaticClass();
+}
+
+TConstArrayView<FAssetCategoryPath> UAssetDefinition_EditableComputeGraph::GetAssetCategories() const
+{
+	static const auto Categories = { EAssetCategoryPaths::Misc };
+	return Categories;
+}
+
+EAssetCommandResult UAssetDefinition_EditableComputeGraph::OpenAssets(const FAssetOpenArgs& OpenArgs) const
+{
+	for (UEditableComputeGraph* EditableComputeGraph : OpenArgs.LoadObjects<UEditableComputeGraph>())
+	{
+		FEditableComputeGraphEditorToolkit::Create(EditableComputeGraph, OpenArgs.GetToolkitMode(), OpenArgs.ToolkitHost);
+	}
+	return EAssetCommandResult::Handled;
+}
diff --git a/Engine/Plugins/Runtime/ComputeFramework/Source/EditableComputeGraphEditor/Private/EditableComputeGraphDetailCustomization.cpp b/Engine/Plugins/Runtime/ComputeFramework/Source/EditableComputeGraphEditor/Private/EditableComputeGraphDetailCustomization.cpp
new file mode 100644
--- /dev/null
+++ b/Engine/Plugins/Runtime/ComputeFramework/Source/EditableComputeGraphEditor/Private/EditableComputeGraphDetailCustomization.cpp
@@ -0,0 +1,267 @@
+// Copyright Epic Games, Inc. All Rights Reserved.
+
+#include "ComputeFramework/EditableComputeGraphDetailCustomization.h"
+
+#include "ComputeFramework/ComputeGraphNavigator.h"
+#include "ComputeFramework/EditableComputeGraph.h"
+#include "ComputeFramework/HlslExternalPinParser.h"
+#include "DetailLayoutBuilder.h"
+#include "DetailWidgetRow.h"
+#include "PropertyHandle.h"
+#include "Styling/AppStyle.h"
+#include "UObject/UnrealType.h"
+#include "Widgets/Input/SComboBox.h"
+#include "Widgets/Text/STextBlock.h"
+
+#define LOCTEXT_NAMESPACE "ComputeFrameworkEditor"
+
+TSharedRef<IDetailCustomization> FEditableComputeGraphDetailCustomization::MakeInstance(TSharedPtr<FComputeGraphEditorSelection> InSelection)
+{
+	TSharedRef<FEditableComputeGraphDetailCustomization> Instance = MakeShared<FEditableComputeGraphDetailCustomization>();
+	Instance->Selection = MoveTemp(InSelection);
+	return Instance;
+}
+
+void FEditableComputeGraphDetailCustomization::CustomizeDetails(IDetailLayoutBuilder& DetailBuilder)
+{
+	// Always hide the raw asset-level categories. We manage display ourselves.
+	DetailBuilder.HideCategory(TEXT("Graph"));
+	DetailBuilder.HideCategory(TEXT("Kernel"));
+	DetailBuilder.HideCategory(TEXT("BindingObject"));
+	DetailBuilder.HideCategory(TEXT("DataInterface"));
+
+	if (!Selection || Selection->Kind == EComputeGraphItemKind::None || Selection->Index == INDEX_NONE)
+	{
+		// Nothing selected — show a placeholder message.
+		DetailBuilder.EditCategory(TEXT("Selection"), LOCTEXT("SelectionCategory", " "))
+			.AddCustomRow(FText::GetEmpty())
+			[
+				SNew(STextBlock)
+				.Text(LOCTEXT("NothingSelected", "Select an item from the navigator."))
+				.Font(DetailBuilder.GetDetailFont())
+				.ColorAndOpacity(FSlateColor::UseSubduedForeground())
+			];
+		return;
+	}
+
+	TSharedRef<IPropertyHandle> GraphDescHandle = DetailBuilder.GetProperty(FName("GraphDescription"));
+	if (!GraphDescHandle->IsValidHandle())
+	{
+		return;
+	}
+
+	FName ArrayPropName;
+	FText CategoryLabel;
+	switch (Selection->Kind)
+	{
+	case EComputeGraphItemKind::Kernel:
+		ArrayPropName = FName("Kernels");
+		CategoryLabel = LOCTEXT("KernelCategory", "Kernel");
+		break;
+	case EComputeGraphItemKind::Interface:
+		ArrayPropName = FName("DataInterfaces");
+		CategoryLabel = LOCTEXT("InterfaceCategory", "Data Interface");
+		break;
+	case EComputeGraphItemKind::BindingObject:
+		ArrayPropName = FName("BindingObjects");
+		CategoryLabel = LOCTEXT("BindingObjectCategory", "Binding Object");
+		break;
+	default:
+		return;
+	}
+
+	TSharedPtr<IPropertyHandle> ArrayHandle = GraphDescHandle->GetChildHandle(ArrayPropName);
+	if (!ArrayHandle.IsValid() || !ArrayHandle->IsValidHandle())
+	{
+		return;
+	}
+
+	TSharedPtr<IPropertyHandle> ElemHandle = ArrayHandle->GetChildHandle(uint32(Selection->Index));
+	if (!ElemHandle.IsValid() || !ElemHandle->IsValidHandle())
+	{
+		return;
+	}
+
+	IDetailCategoryBuilder& Category = DetailBuilder.EditCategory(TEXT("SelectedItem"), CategoryLabel);
+
+	TArray<TWeakObjectPtr<UObject>> ObjectsBeingCustomized;
+	DetailBuilder.GetObjectsBeingCustomized(ObjectsBeingCustomized);
+	UEditableComputeGraph* Graph = (ObjectsBeingCustomized.Num() == 1) ? Cast<UEditableComputeGraph>(ObjectsBeingCustomized[0].Get()) : nullptr;
+
+	// For the EntryPoint combo we need access to the SourceText handle.
+	TSharedPtr<IPropertyHandle> SourceTextHandle;
+	if (Selection->Kind == EComputeGraphItemKind::Kernel)
+	{
+		SourceTextHandle = ElemHandle->GetChildHandle(FName("SourceText"));
+	}
+
+	uint32 NumChildren = 0;
+	ElemHandle->GetNumChildren(NumChildren);
+	for (uint32 i = 0; i < NumChildren; ++i)
+	{
+		TSharedPtr<IPropertyHandle> ChildHandle = ElemHandle->GetChildHandle(i);
+		if (!ChildHandle.IsValid() || !ChildHandle->IsValidHandle()) 
+		{
+			continue;
+		}
+
+		const FName PropName = ChildHandle->GetProperty() ? ChildHandle->GetProperty()->GetFName() : NAME_None;
+
+		// Skip internal / redundant fields.
+		if (PropName == FName("bOrphaned"))
+		{
+			continue;
+		}
+
+		// SourceText is always visible in the HLSL editor panel.
+		if (Selection->Kind == EComputeGraphItemKind::Kernel && PropName == FName("SourceText")) 
+		{
+			continue;
+		}
+
+		// BindingObjectName combo box populated from BindingObjects.
+		if (Selection->Kind == EComputeGraphItemKind::Interface && PropName == FName("BindingObjectName"))
+		{
+			TSharedPtr<IPropertyHandle> BindingHandle = ChildHandle;
+			TSharedRef<TArray<TSharedPtr<FName>>> Options = MakeShared<TArray<TSharedPtr<FName>>>();
+
+			Category.AddCustomRow(LOCTEXT("BindingObjectName", "Binding Object"))
+				.NameContent()
+				[
+					BindingHandle->CreatePropertyNameWidget()
+				]
+				.ValueContent()
+				.MinDesiredWidth(200.f)
+				[
+					SNew(SComboBox<TSharedPtr<FName>>)
+					.OptionsSource(&Options.Get())
+					.OnComboBoxOpening_Lambda([Options, Graph]()
+					{
+						Options->Reset();
+						if (IsValid(Graph))
+						{
+							for (FComputeGraphDataBindingObjectDesc const& Desc : Graph->GetGraphDescription().BindingObjects)
+							{
+								if (!Desc.Name.IsNone())
+								{
+									Options->Add(MakeShared<FName>(Desc.Name));
+								}
+							}
+						}
+					})
+					.OnSelectionChanged_Lambda([BindingHandle](TSharedPtr<FName> NewValue, ESelectInfo::Type)
+					{
+						if (BindingHandle.IsValid() && NewValue.IsValid())
+						{
+							BindingHandle->SetValue(*NewValue);
+						}
+					})
+					.OnGenerateWidget_Lambda([](TSharedPtr<FName> Item) -> TSharedRef<SWidget>
+					{
+						return SNew(STextBlock)
+							.Text(FText::FromName(Item.IsValid() ? *Item : NAME_None))
+							.Font(FAppStyle::GetFontStyle("PropertyWindow.NormalFont"));
+					})
+					[
+						SNew(STextBlock)
+						.Text_Lambda([BindingHandle]() -> FText
+						{
+							FName Value;
+							if (BindingHandle.IsValid()
+								&& BindingHandle->GetValue(Value) == FPropertyAccess::Success
+								&& !Value.IsNone())
+							{
+								return FText::FromName(Value);
+							}
+							return LOCTEXT("SelectBindingObject", "(select)");
+						})
+						.Font(FAppStyle::GetFontStyle("PropertyWindow.NormalFont"))
+					]
+				];
+			continue;
+		}
+
+		// For the DataInterface Settings flatten the associated UComputeDataInterface subclass properties directly.
+		if (Selection->Kind == EComputeGraphItemKind::Interface && PropName == FName("Settings"))
+		{
+			UObject* SettingsObj = nullptr;
+			if (ChildHandle->GetValue(SettingsObj) == FPropertyAccess::Success && IsValid(SettingsObj))
+			{
+				TArray<UObject*> SettingsObjects = { SettingsObj };
+				for (TFieldIterator<FProperty> It(SettingsObj->GetClass()); It; ++It)
+				{
+					if (It->HasAllPropertyFlags(CPF_Edit))
+					{
+						Category.AddExternalObjectProperty(SettingsObjects, It->GetFName());
+					}
+				}
+			}
+			continue;
+		}
+
+		// The EntryPoint combo box is populated by parsing function definitions from the HLSL source.
+		if (Selection->Kind == EComputeGraphItemKind::Kernel && PropName == FName("EntryPoint"))
+		{
+			TSharedPtr<IPropertyHandle> EntryPointHandle = ChildHandle;
+			TSharedRef<TArray<TSharedPtr<FString>>> Options = MakeShared<TArray<TSharedPtr<FString>>>();
+
+			Category.AddCustomRow(LOCTEXT("EntryPoint", "Entry Point"))
+				.NameContent()
+				[
+					EntryPointHandle->CreatePropertyNameWidget()
+				]
+				.ValueContent()
+				.MinDesiredWidth(200.f)
+				[
+					SNew(SComboBox<TSharedPtr<FString>>)
+					.OptionsSource(&Options.Get())
+					.OnComboBoxOpening_Lambda([Options, SourceTextHandle]()
+					{
+						Options->Reset();
+						FString SourceText;
+						if (SourceTextHandle.IsValid())
+						{
+							SourceTextHandle->GetValue(SourceText);
+						}
+						for (FString const& FnName : FHlslExternalPinParser::FindFunctionDefinitions(SourceText))
+						{
+							Options->Add(MakeShared<FString>(FnName));
+						}
+					})
+					.OnSelectionChanged_Lambda([EntryPointHandle](TSharedPtr<FString> NewValue, ESelectInfo::Type)
+					{
+						if (EntryPointHandle.IsValid() && NewValue.IsValid())
+						{
+							EntryPointHandle->SetValue(*NewValue);
+						}
+					})
+					.OnGenerateWidget_Lambda([](TSharedPtr<FString> Item) -> TSharedRef<SWidget>
+					{
+						return SNew(STextBlock)
+							.Text(FText::FromString(Item.IsValid() ? *Item : FString()))
+							.Font(FAppStyle::GetFontStyle("PropertyWindow.NormalFont"));
+					})
+					[
+						SNew(STextBlock)
+						.Text_Lambda([EntryPointHandle]() -> FText
+						{
+							FString Value;
+							if (EntryPointHandle.IsValid()
+								&& EntryPointHandle->GetValue(Value) == FPropertyAccess::Success
+								&& !Value.IsEmpty())
+							{
+								return FText::FromString(Value);
+							}
+							return LOCTEXT("SelectEntryPoint", "(select)");
+						})
+						.Font(FAppStyle::GetFontStyle("PropertyWindow.NormalFont"))
+					]
+				];
+			continue;
+		}
+
+		Category.AddProperty(ChildHandle.ToSharedRef());
+	}
+}
+
+#undef LOCTEXT_NAMESPACE
diff --git a/Engine/Plugins/Runtime/ComputeFramework/Source/EditableComputeGraphEditor/Private/EditableComputeGraphEditorModule.cpp b/Engine/Plugins/Runtime/ComputeFramework/Source/EditableComputeGraphEditor/Private/EditableComputeGraphEditorModule.cpp
new file mode 100644
--- /dev/null
+++ b/Engine/Plugins/Runtime/ComputeFramework/Source/EditableComputeGraphEditor/Private/EditableComputeGraphEditorModule.cpp
@@ -0,0 +1,27 @@
+// Copyright Epic Games, Inc. All Rights Reserved.
+
+#include "ComputeFramework/EditableComputeGraph.h"
+#include "ComputeFramework/KernelPinCustomization.h"
+#include "Modules/ModuleManager.h"
+#include "PropertyEditorModule.h"
+
+class FEditableComputeGraphEditorModule : public IModuleInterface
+{
+public:
+	void StartupModule() override
+	{
+		FPropertyEditorModule& PropertyModule = FModuleManager::LoadModuleChecked<FPropertyEditorModule>("PropertyEditor");
+		PropertyModule.RegisterCustomPropertyTypeLayout(FKernelPin::StaticStruct()->GetFName(),	FOnGetPropertyTypeCustomizationInstance::CreateStatic(&FKernelPinCustomization::MakeInstance));
+	}
+
+	void ShutdownModule() override
+	{
+		if (FModuleManager::Get().IsModuleLoaded("PropertyEditor"))
+		{
+			FPropertyEditorModule& PropertyModule = FModuleManager::GetModuleChecked<FPropertyEditorModule>("PropertyEditor");
+			PropertyModule.UnregisterCustomPropertyTypeLayout(FKernelPin::StaticStruct()->GetFName());
+		}
+	}
+};
+
+IMPLEMENT_MODULE(FEditableComputeGraphEditorModule, EditableComputeGraphEditor)
diff --git a/Engine/Plugins/Runtime/ComputeFramework/Source/EditableComputeGraphEditor/Private/EditableComputeGraphEditorToolkit.cpp b/Engine/Plugins/Runtime/ComputeFramework/Source/EditableComputeGraphEditor/Private/EditableComputeGraphEditorToolkit.cpp
new file mode 100644
--- /dev/null
+++ b/Engine/Plugins/Runtime/ComputeFramework/Source/EditableComputeGraphEditor/Private/EditableComputeGraphEditorToolkit.cpp
@@ -0,0 +1,1043 @@
+// Copyright Epic Games, Inc. All Rights Reserved.
+
+#include "ComputeFramework/EditableComputeGraphEditorToolkit.h"
+
+#include "ComputeFramework/ComputeKernelCompileResult.h"
+#include "ComputeFramework/EditableComputeGraph.h"
+#include "ComputeFramework/ComputeGraphHlslEditor.h"
+#include "ComputeFramework/ComputeGraphNavigator.h"
+#include "ComputeFramework/ComputeGraphView.h"
+#include "ComputeFramework/HlslExternalPinParser.h"
+#include "Framework/Docking/TabManager.h"
+#include "Framework/MultiBox/MultiBoxBuilder.h"
+#include "IDetailsView.h"
+#include "Modules/ModuleManager.h"
+#include "PropertyEditorModule.h"
+#include "ScopedTransaction.h"
+#include "Styling/AppStyle.h"
+#include "Widgets/Docking/SDockTab.h"
+#include "Widgets/Layout/SBorder.h"
+#include "Widgets/Layout/SScrollBox.h"
+#include "Widgets/Text/STextBlock.h"
+#include "Widgets/Views/SListView.h"
+#include "Widgets/Views/STableRow.h"
+
+#define LOCTEXT_NAMESPACE "ComputeFrameworkEditor"
+
+const FName FEditableComputeGraphEditorToolkit::TabId_Navigator = TEXT("TID_Navigator");
+const FName FEditableComputeGraphEditorToolkit::TabId_HlslEditor = TEXT("TID_HlslEditor");
+const FName FEditableComputeGraphEditorToolkit::TabId_GraphView = TEXT("TID_GraphView");
+const FName FEditableComputeGraphEditorToolkit::TabId_Details = TEXT("TID_Details");
+const FName FEditableComputeGraphEditorToolkit::TabId_Output = TEXT("TID_Output");
+
+static void RenameInterfaceInPins(FComputeGraphDesc& Desc, FName OldName, FName NewName)
+{
+	for (FComputeGraphKernelDesc& KernelDesc : Desc.Kernels)
+	{
+		for (FKernelPin& Pin : KernelDesc.Inputs)
+		{
+			if (Pin.DataInterfaceName == OldName)
+			{
+				Pin.DataInterfaceName = NewName;
+			}
+		}
+		for (FKernelPin& Pin : KernelDesc.Outputs)
+		{
+			if (Pin.DataInterfaceName == OldName)
+			{
+				Pin.DataInterfaceName = NewName;
+			}
+		}
+	}
+}
+
+static void RenameBindingObjectInInterfaces(FComputeGraphDesc& Desc, FName OldName, FName NewName)
+{
+	for (FComputeGraphDataInterfaceDesc& IfaceDesc : Desc.DataInterfaces)
+	{
+		if (IfaceDesc.BindingObjectName == OldName)
+		{
+			IfaceDesc.BindingObjectName = NewName;
+		}
+	}
+}
+
+TSharedRef<FEditableComputeGraphEditorToolkit> FEditableComputeGraphEditorToolkit::Create(UEditableComputeGraph* InAsset, EToolkitMode::Type InMode, TSharedPtr<IToolkitHost> InToolkitHost)
+{
+	TSharedRef<FEditableComputeGraphEditorToolkit> Toolkit = MakeShared<FEditableComputeGraphEditorToolkit>();
+	Toolkit->Init(InAsset, InMode, InToolkitHost);
+	return Toolkit;
+}
+
+void FEditableComputeGraphEditorToolkit::Init(UEditableComputeGraph* InAsset, EToolkitMode::Type InMode, TSharedPtr<IToolkitHost> InToolkitHost)
+{
+	Asset = InAsset;
+	EditorSelection = MakeShared<FComputeGraphEditorSelection>();
+
+	// Details view.
+	FPropertyEditorModule& PropertyModule = FModuleManager::LoadModuleChecked<FPropertyEditorModule>("PropertyEditor");
+	FDetailsViewArgs DetailsArgs;
+	DetailsArgs.bHideSelectionTip = true;
+	DetailsView = PropertyModule.CreateDetailView(DetailsArgs);
+	DetailsView->SetObject(InAsset);
+
+	// Refresh the navigator whenever any property changes in the details panel.
+	// Also apply cross-reference fixups when an Interface or BindingObject Name is edited.
+	DetailsView->OnFinishedChangingProperties().AddLambda([this](FPropertyChangedEvent const&)
+	{
+		if (Asset && EditorSelection && !CachedSelectedItemName.IsNone())
+		{
+			FComputeGraphDesc& Desc = Asset->GetGraphDescription();
+
+			if (EditorSelection->Kind == EComputeGraphItemKind::Interface && Desc.DataInterfaces.IsValidIndex(EditorSelection->Index))
+			{
+				const FName NewName = Desc.DataInterfaces[EditorSelection->Index].Name;
+				if (NewName != CachedSelectedItemName)
+				{
+					RenameInterfaceInPins(Desc, CachedSelectedItemName, NewName);
+					CachedSelectedItemName = NewName;
+					bDirty = true;
+				}
+			}
+			else if (EditorSelection->Kind == EComputeGraphItemKind::BindingObject && Desc.BindingObjects.IsValidIndex(EditorSelection->Index))
+			{
+				const FName NewName = Desc.BindingObjects[EditorSelection->Index].Name;
+				if (NewName != CachedSelectedItemName)
+				{
+					RenameBindingObjectInInterfaces(Desc, CachedSelectedItemName, NewName);
+					CachedSelectedItemName = NewName;
+					bDirty = true;
+				}
+			}
+		}
+
+		if (NavigatorWidget.IsValid())
+		{
+			NavigatorWidget->Refresh();
+		}
+		if (GraphViewWidget.IsValid())
+		{
+			GraphViewWidget->Refresh();
+		}
+	});
+
+	// Subscribe to compile output from the asset.
+	CompileOutputHandle = InAsset->OnCompileOutputChanged.AddRaw(this, &FEditableComputeGraphEditorToolkit::OnCompileOutputChanged);
+
+	// Register per-instance customization so the details panel shows only the selected navigator item rather than the full asset.
+	DetailsView->RegisterInstancedCustomPropertyLayout(
+		UEditableComputeGraph::StaticClass(),
+		FOnGetDetailCustomizationInstance::CreateStatic(&FEditableComputeGraphDetailCustomization::MakeInstance, EditorSelection));
+
+	// Tab layout.
+	TSharedRef<FTabManager::FLayout> Layout = FTabManager::NewLayout("EditableComputeGraphEditor_v4")
+		->AddArea
+		(
+			FTabManager::NewPrimaryArea()
+			->SetOrientation(Orient_Horizontal)
+			->Split
+			(
+				FTabManager::NewStack()
+				->SetSizeCoefficient(0.22f)
+				->AddTab(TabId_Navigator, ETabState::OpenedTab)
+				->SetHideTabWell(true)
+			)
+			->Split
+			(
+				FTabManager::NewSplitter()
+				->SetOrientation(Orient_Vertical)
+				->SetSizeCoefficient(0.50f)
+				->Split
+				(
+					FTabManager::NewStack()
+					->SetSizeCoefficient(0.50f)
+					->AddTab(TabId_HlslEditor, ETabState::OpenedTab)
+					->SetHideTabWell(true)
+				)
+				->Split
+				(
+					FTabManager::NewStack()
+					->SetSizeCoefficient(0.25f)
+					->AddTab(TabId_GraphView, ETabState::OpenedTab)
+					->SetHideTabWell(true)
+				)
+				->Split
+				(
+					FTabManager::NewStack()
+					->SetSizeCoefficient(0.25f)
+					->AddTab(TabId_Output, ETabState::OpenedTab)
+					->SetHideTabWell(true)
+				)
+			)
+			->Split
+			(
+				FTabManager::NewStack()
+				->SetSizeCoefficient(0.28f)
+				->AddTab(TabId_Details, ETabState::OpenedTab)
+				->SetHideTabWell(true)
+			)
+		);
+
+	constexpr bool bCreateDefaultStandaloneMenu = true;
+	constexpr bool bCreateDefaultToolbar = true;
+	FAssetEditorToolkit::InitAssetEditor(
+		InMode,
+		InToolkitHost,
+		TEXT("EditableComputeGraphEditorApp"),
+		Layout,
+		bCreateDefaultStandaloneMenu,
+		bCreateDefaultToolbar,
+		InAsset);
+
+	ExtendToolbar();
+	RegenerateMenusAndToolbars();
+}
+
+FName FEditableComputeGraphEditorToolkit::GetToolkitFName() const
+{
+	return TEXT("EditableComputeGraphEditor");
+}
+
+FText FEditableComputeGraphEditorToolkit::GetBaseToolkitName() const
+{
+	return LOCTEXT("ToolkitName", "Compute Graph Editor");
+}
+
+FString FEditableComputeGraphEditorToolkit::GetWorldCentricTabPrefix() const
+{
+	return TEXT("ComputeGraph ");
+}
+
+FLinearColor FEditableComputeGraphEditorToolkit::GetWorldCentricTabColorScale() const
+{
+	return FLinearColor(0.0f, 0.5f, 0.5f, 0.5f);
+}
+
+void FEditableComputeGraphEditorToolkit::AddReferencedObjects(FReferenceCollector& Collector)
+{
+	Collector.AddReferencedObject(Asset);
+}
+
+FString FEditableComputeGraphEditorToolkit::GetReferencerName() const
+{
+	return TEXT("FEditableComputeGraphEditorToolkit");
+}
+
+void FEditableComputeGraphEditorToolkit::RegisterTabSpawners(TSharedRef<FTabManager> const& InTabManager)
+{
+	FAssetEditorToolkit::RegisterTabSpawners(InTabManager);
+
+	WorkspaceMenuCategory = InTabManager->AddLocalWorkspaceMenuCategory(LOCTEXT("WorkspaceMenuCategory", "Compute Graph Editor"));
+
+	InTabManager->RegisterTabSpawner(TabId_Navigator, FOnSpawnTab::CreateSP(this, &FEditableComputeGraphEditorToolkit::SpawnTab_Navigator))
+		.SetDisplayName(LOCTEXT("NavigatorTab", "Navigator"))
+		.SetGroup(WorkspaceMenuCategory.ToSharedRef());
+
+	InTabManager->RegisterTabSpawner(TabId_HlslEditor, FOnSpawnTab::CreateSP(this, &FEditableComputeGraphEditorToolkit::SpawnTab_HlslEditor))
+		.SetDisplayName(LOCTEXT("HlslEditorTab", "HLSL Editor"))
+		.SetGroup(WorkspaceMenuCategory.ToSharedRef());
+
+	InTabManager->RegisterTabSpawner(TabId_GraphView, FOnSpawnTab::CreateSP(this, &FEditableComputeGraphEditorToolkit::SpawnTab_GraphView))
+		.SetDisplayName(LOCTEXT("GraphViewTab", "Graph View"))
+		.SetGroup(WorkspaceMenuCategory.ToSharedRef());
+
+	InTabManager->RegisterTabSpawner(TabId_Details, FOnSpawnTab::CreateSP(this, &FEditableComputeGraphEditorToolkit::SpawnTab_Details))
+		.SetDisplayName(LOCTEXT("DetailsTab", "Details"))
+		.SetGroup(WorkspaceMenuCategory.ToSharedRef());
+
+	InTabManager->RegisterTabSpawner(TabId_Output, FOnSpawnTab::CreateSP(this, &FEditableComputeGraphEditorToolkit::SpawnTab_Output))
+		.SetDisplayName(LOCTEXT("OutputTab", "Output"))
+		.SetGroup(WorkspaceMenuCategory.ToSharedRef());
+}
+
+void FEditableComputeGraphEditorToolkit::UnregisterTabSpawners(TSharedRef<FTabManager> const& InTabManager)
+{
+	FAssetEditorToolkit::UnregisterTabSpawners(InTabManager);
+	InTabManager->UnregisterTabSpawner(TabId_Navigator);
+	InTabManager->UnregisterTabSpawner(TabId_HlslEditor);
+	InTabManager->UnregisterTabSpawner(TabId_GraphView);
+	InTabManager->UnregisterTabSpawner(TabId_Details);
+	InTabManager->UnregisterTabSpawner(TabId_Output);
+
+	if (Asset != nullptr)
+	{
+		Asset->OnCompileOutputChanged.Remove(CompileOutputHandle);
+	}
+}
+
+TSharedRef<SDockTab> FEditableComputeGraphEditorToolkit::SpawnTab_Navigator(FSpawnTabArgs const& Args)
+{
+	NavigatorWidget = SNew(SComputeGraphNavigator)
+		.Asset(Asset)
+		.OnItemSelected(FOnComputeGraphItemSelected::CreateSP(this, &FEditableComputeGraphEditorToolkit::OnItemSelected))
+		.OnAddItem(FOnComputeGraphAddItem::CreateSP(this, &FEditableComputeGraphEditorToolkit::OnAddItem))
+		.OnDeleteItem(FOnComputeGraphDeleteItem::CreateSP(this, &FEditableComputeGraphEditorToolkit::OnDeleteItem))
+		.OnDuplicateItem(FOnComputeGraphDuplicateItem::CreateSP(this, &FEditableComputeGraphEditorToolkit::OnDuplicateItem))
+		.OnRenameItem(FOnComputeGraphRenameItem::CreateSP(this, &FEditableComputeGraphEditorToolkit::OnRenameItem))
+		.OnKernelSelected(FOnComputeGraphKernelSelected::CreateSP(this, &FEditableComputeGraphEditorToolkit::OnKernelSelected));
+
+	return SNew(SDockTab)
+		.TabRole(ETabRole::PanelTab)
+		[
+			NavigatorWidget.ToSharedRef()
+		];
+}
+
+TSharedRef<SDockTab> FEditableComputeGraphEditorToolkit::SpawnTab_HlslEditor(FSpawnTabArgs const& Args)
+{
+	HlslEditorWidget = SNew(SComputeGraphHlslEditor)
+		.Asset(Asset)
+		.OnTextCommitted_WithKernel(FOnComputeGraphHlslTextCommitted::CreateSP(this, &FEditableComputeGraphEditorToolkit::OnHlslTextCommitted))
+		.OnTextChanged(FOnComputeGraphHlslTextChanged::CreateLambda([this]() { bDirty = true; }))
+		.OnCompileRequested(FOnComputeGraphHlslCompileRequested::CreateSP(this, &FEditableComputeGraphEditorToolkit::CompileGraph));
+
+	return SNew(SDockTab)
+		.TabRole(ETabRole::PanelTab)
+		[
+			HlslEditorWidget.ToSharedRef()
+		];
+}
+
+TSharedRef<SDockTab> FEditableComputeGraphEditorToolkit::SpawnTab_GraphView(FSpawnTabArgs const& Args)
+{
+	GraphViewWidget = SNew(SComputeGraphView)
+		.Asset(Asset)
+		.OnNodeClicked(this, &FEditableComputeGraphEditorToolkit::OnGraphNodeClicked);
+
+	return SNew(SDockTab)
+		.TabRole(ETabRole::PanelTab)
+		[
+			SNew(SScrollBox)
+			.Orientation(Orient_Horizontal)
+			+ SScrollBox::Slot()
+			[
+				SNew(SScrollBox)
+				.Orientation(Orient_Vertical)
+				+ SScrollBox::Slot()
+				[
+					GraphViewWidget.ToSharedRef()
+				]
+			]
+		];
+}
+
+TSharedRef<SDockTab> FEditableComputeGraphEditorToolkit::SpawnTab_Details(FSpawnTabArgs const& Args)
+{
+	return SNew(SDockTab)
+		.TabRole(ETabRole::PanelTab)
+		[
+			DetailsView.ToSharedRef()
+		];
+}
+
+TSharedRef<SDockTab> FEditableComputeGraphEditorToolkit::SpawnTab_Output(FSpawnTabArgs const& Args)
+{
+	SAssignNew(OutputListView, SListView<TSharedPtr<FComputeKernelCompileMessage>>)
+		.ListItemsSource(&OutputMessages)
+		.OnGenerateRow(this, &FEditableComputeGraphEditorToolkit::GenerateOutputRow)
+		.OnSelectionChanged(this, &FEditableComputeGraphEditorToolkit::OnOutputSelectionChanged)
+		.SelectionMode(ESelectionMode::Single);
+
+	return SNew(SDockTab)
+		.TabRole(ETabRole::PanelTab)
+		[
+			SNew(SBorder)
+			.BorderImage(FAppStyle::Get().GetBrush("ToolPanel.GroupBorder"))
+			.Padding(4.f)
+			[
+				OutputListView.ToSharedRef()
+			]
+		];
+}
+
+TSharedRef<ITableRow> FEditableComputeGraphEditorToolkit::GenerateOutputRow(
+	TSharedPtr<FComputeKernelCompileMessage> Message,
+	TSharedRef<STableViewBase> const& OwnerTable)
+{
+	FSlateColor Color = FSlateColor::UseForeground();
+	FString Text;
+
+	if (Message.IsValid())
+	{
+		switch (Message->Type)
+		{
+		case FComputeKernelCompileMessage::EMessageType::Error:
+			Color = FSlateColor(FLinearColor(1.f, 0.3f, 0.3f));
+			break;
+		case FComputeKernelCompileMessage::EMessageType::Warning:
+			Color = FSlateColor(FLinearColor(1.f, 0.85f, 0.2f));
+			break;
+		default:
+			break;
+		}
+
+		Text = Message->Line >= 0
+			? FString::Printf(TEXT("(%d): %s"), Message->Line, *Message->Text)
+			: Message->Text;
+	}
+
+	return SNew(STableRow<TSharedPtr<FComputeKernelCompileMessage>>, OwnerTable)
+		.Padding(FMargin(4.f, 1.f))
+		[
+			SNew(STextBlock)
+			.Text(FText::FromString(Text))
+			.Font(FAppStyle::GetFontStyle("PropertyWindow.NormalFont"))
+			.ColorAndOpacity(Color)
+		];
+}
+
+void FEditableComputeGraphEditorToolkit::OnCompileOutputChanged(TArray<FComputeKernelCompileMessage> const& NewMessages)
+{
+	if (NewMessages.IsEmpty())
+	{
+		OutputMessages.Reset();
+	}
+	else
+	{
+		for (FComputeKernelCompileMessage const& Msg : NewMessages)
+		{
+			OutputMessages.Add(MakeShared<FComputeKernelCompileMessage>(Msg));
+		}
+	}
+	if (OutputListView.IsValid())
+	{
+		OutputListView->RequestListRefresh();
+	}
+}
+
+void FEditableComputeGraphEditorToolkit::OnOutputSelectionChanged(TSharedPtr<FComputeKernelCompileMessage> Selected, ESelectInfo::Type /*SelectInfo*/)
+{
+	if (!Selected.IsValid() || !Asset || !NavigatorWidget.IsValid())
+	{
+		return;
+	}
+
+	// Parse the [KernelName] prefix added by OnKernelCompilationComplete.
+	if (!Selected->Text.StartsWith(TEXT("[")))
+	{
+		return;
+	}
+	int32 CloseBracket = INDEX_NONE;
+	Selected->Text.FindChar(TEXT(']'), CloseBracket);
+	if (CloseBracket <= 1)
+	{
+		return;
+	}
+	const FString ParsedName = Selected->Text.Mid(1, CloseBracket - 1);
+
+	TArray<FComputeGraphKernelDesc> const& Kernels = Asset->GetGraphDescription().Kernels;
+	const int32 KernelIdx = Kernels.IndexOfByPredicate([&](FComputeGraphKernelDesc const& K) { return K.Name.ToString() == ParsedName; });
+	if (KernelIdx != INDEX_NONE)
+	{
+		NavigatorWidget->SetSelectedItem(EComputeGraphItemKind::Kernel, KernelIdx);
+	}
+}
+
+void FEditableComputeGraphEditorToolkit::ExtendToolbar()
+{
+	TSharedPtr<FExtender> Extender = MakeShared<FExtender>();
+	Extender->AddToolBarExtension(
+		"Asset",
+		EExtensionHook::After,
+		GetToolkitCommands(),
+		FToolBarExtensionDelegate::CreateLambda([this](FToolBarBuilder& ToolbarBuilder)
+		{
+			ToolbarBuilder.AddToolBarButton(
+				FUIAction(
+					FExecuteAction::CreateSP(this, &FEditableComputeGraphEditorToolkit::CompileGraph),
+					FCanExecuteAction::CreateLambda([this]() { return GetCompileState() != ECompileState::UpToDate; })),
+				NAME_None,
+				LOCTEXT("CompileButton", "Compile"),
+				LOCTEXT("CompileButtonTooltip", "Rebuild the compute graph and trigger shader compilation."),
+				TAttribute<FSlateIcon>(this, &FEditableComputeGraphEditorToolkit::GetCompileStatusIcon));
+		}));
+	AddToolbarExtender(Extender);
+}
+
+FEditableComputeGraphEditorToolkit::ECompileState FEditableComputeGraphEditorToolkit::GetCompileState() const
+{
+	if (!bDirty)
+	{
+		return ECompileState::UpToDate;
+	}
+
+	// Validate: Every DataInterface must have a BindingObject set, and every kernel pin must have all three name fields filled.
+	if (Asset)
+	{
+		FComputeGraphDesc const& Desc = Asset->GetGraphDescription();
+
+		for (FComputeGraphDataInterfaceDesc const& Iface : Desc.DataInterfaces)
+		{
+			if (Iface.BindingObjectName.IsNone())
+			{
+				return ECompileState::Broken;
+			}
+		}
+
+		for (FComputeGraphKernelDesc const& Kernel : Desc.Kernels)
+		{
+			auto PinIncomplete = [](FKernelPin const& Pin) -> bool
+			{
+				return Pin.KernelFunctionName.IsEmpty()	|| Pin.DataInterfaceName.IsNone() || Pin.DataInterfaceFunctionName.IsEmpty();
+			};
+
+			for (FKernelPin const& Pin : Kernel.Inputs)
+			{
+				if (PinIncomplete(Pin))
+				{
+					return ECompileState::Broken;
+				}
+			}
+			for (FKernelPin const& Pin : Kernel.Outputs)
+			{
+				if (PinIncomplete(Pin))
+				{
+					return ECompileState::Broken;
+				}
+			}
+		}
+	}
+
+	return ECompileState::NeedsCompile;
+}
+
+FSlateIcon FEditableComputeGraphEditorToolkit::GetCompileStatusIcon() const
+{
+	static const FName Background(TEXT("AssetEditor.CompileStatus.Background"));
+	static const FName OverlayGood(TEXT("AssetEditor.CompileStatus.Overlay.Good"));
+	static const FName OverlayUnknown(TEXT("AssetEditor.CompileStatus.Overlay.Unknown"));
+	static const FName OverlayError(TEXT("AssetEditor.CompileStatus.Overlay.Error"));
+
+	const FName StyleSet = FAppStyle::GetAppStyleSetName();
+
+	switch (GetCompileState())
+	{
+	case ECompileState::UpToDate:
+		return FSlateIcon(StyleSet, Background, OverlayGood);
+	case ECompileState::Broken:
+		return FSlateIcon(StyleSet, Background, OverlayError);
+	default:
+		return FSlateIcon(StyleSet, Background, OverlayUnknown);
+	}
+}
+
+void FEditableComputeGraphEditorToolkit::OnItemSelected(EComputeGraphItemKind Kind, int32 Index, FName Name)
+{
+	if (EditorSelection)
+	{
+		EditorSelection->Kind = Kind;
+		EditorSelection->Index = Index;
+	}
+
+	// Snapshot the name so OnFinishedChangingProperties can detect in-panel renames.
+	CachedSelectedItemName = NAME_None;
+	if (Asset && (Kind == EComputeGraphItemKind::Interface || Kind == EComputeGraphItemKind::BindingObject))
+	{
+		CachedSelectedItemName = Name;
+	}
+
+	DetailsView->ForceRefresh();
+
+	// Also drive the HLSL editor / graph view depending on selection kind.
+	if (Kind == EComputeGraphItemKind::Kernel)
+	{
+		OnKernelSelected(Name);
+		if (GraphViewWidget.IsValid()) 
+		{ 
+			GraphViewWidget->SetHighlightedInterface(NAME_None); 
+		}
+	}
+	else if (Kind == EComputeGraphItemKind::Interface)
+	{
+		if (GraphViewWidget.IsValid())
+		{
+			GraphViewWidget->SetHighlightedKernel(NAME_None);
+			GraphViewWidget->SetHighlightedInterface(Name);
+		}
+	}
+	else
+	{
+		if (GraphViewWidget.IsValid())
+		{
+			GraphViewWidget->SetHighlightedKernel(NAME_None);
+			GraphViewWidget->SetHighlightedInterface(NAME_None);
+		}
+	}
+}
+
+void FEditableComputeGraphEditorToolkit::OnGraphNodeClicked(FName NodeName, bool bIsKernel)
+{
+	if (Asset == nullptr)
+	{
+		return;
+	}
+
+	FComputeGraphDesc const& Desc = Asset->GetGraphDescription();
+
+	if (bIsKernel)
+	{
+		for (int32 KernelIndex = 0; KernelIndex < Desc.Kernels.Num(); ++KernelIndex)
+		{
+			if (Desc.Kernels[KernelIndex].Name == NodeName)
+			{
+				NavigatorWidget->SetSelectedItem(EComputeGraphItemKind::Kernel, KernelIndex);
+				OnItemSelected(EComputeGraphItemKind::Kernel, KernelIndex, NodeName);
+				return;
+			}
+		}
+	}
+	else
+	{
+		for (int32 DataInterfaceIndex = 0; DataInterfaceIndex < Desc.DataInterfaces.Num(); ++DataInterfaceIndex)
+		{
+			if (Desc.DataInterfaces[DataInterfaceIndex].Name == NodeName)
+			{
+				NavigatorWidget->SetSelectedItem(EComputeGraphItemKind::Interface, DataInterfaceIndex);
+				OnItemSelected(EComputeGraphItemKind::Interface, DataInterfaceIndex, NodeName);
+				return;
+			}
+		}
+	}
+}
+
+void FEditableComputeGraphEditorToolkit::RefreshHlslBoundFunctions(FName KernelName)
+{
+	if (Asset == nullptr || !HlslEditorWidget.IsValid()) 
+	{
+		return;
+	}
+
+	TArray<FString> BoundNames;
+	for (FComputeGraphKernelDesc const& KernelDesc : Asset->GetGraphDescription().Kernels)
+	{
+		if (KernelDesc.Name != KernelName)
+		{
+			continue;
+		}
+		if (!KernelDesc.EntryPoint.IsEmpty())
+		{
+			BoundNames.Add(KernelDesc.EntryPoint);
+		}
+		for (FKernelPin const& Pin : KernelDesc.Inputs)
+		{
+			if (!Pin.KernelFunctionName.IsEmpty())
+			{
+				BoundNames.Add(Pin.KernelFunctionName);
+			}
+		}
+		for (FKernelPin const& Pin : KernelDesc.Outputs)
+		{
+			if (!Pin.KernelFunctionName.IsEmpty())
+			{
+				BoundNames.Add(Pin.KernelFunctionName);
+			}
+		}
+		break;
+	}
+	HlslEditorWidget->SetBoundFunctions(BoundNames);
+}
+
+void FEditableComputeGraphEditorToolkit::OnKernelSelected(FName KernelName)
+{
+	if (HlslEditorWidget.IsValid())
+	{
+		HlslEditorWidget->SetActiveKernel(KernelName);
+		RefreshHlslBoundFunctions(KernelName);
+	}
+	if (GraphViewWidget.IsValid())
+	{
+		GraphViewWidget->SetHighlightedKernel(KernelName);
+	}
+}
+
+FName FEditableComputeGraphEditorToolkit::GenerateUniqueName(TArray<FName> const& Existing, FStringView Base)
+{
+	FString Candidate(Base);
+	int32 Suffix = 0;
+	while (Existing.Contains(FName(*Candidate)))
+	{
+		Candidate = FString::Printf(TEXT("%s%d"), *FString(Base), ++Suffix);
+	}
+	return FName(*Candidate);
+}
+
+void FEditableComputeGraphEditorToolkit::OnAddItem(EComputeGraphItemKind Kind)
+{
+	if (Asset == nullptr) 
+	{
+		return;
+	}
+
+	FScopedTransaction Transaction(LOCTEXT("AddItem", "Add Compute Graph Item"));
+	Asset->Modify();
+
+	FComputeGraphDesc& Desc = Asset->GetGraphDescription();
+	int32 NewIndex = INDEX_NONE;
+
+	switch (Kind)
+	{
+	case EComputeGraphItemKind::Kernel:
+	{
+		TArray<FName> Existing;
+		for (FComputeGraphKernelDesc const& K : Desc.Kernels)
+		{
+			Existing.Add(K.Name);
+		}
+
+		// Add stub kernel main function.
+		FComputeGraphKernelDesc& New = Desc.Kernels.AddDefaulted_GetRef();
+		New.Name = GenerateUniqueName(Existing, TEXT("NewKernel"));
+		New.EntryPoint = TEXT("Main");
+		New.GroupSize = FIntVector(64, 1, 1);
+		New.SourceText = TEXT("void Main(uint3 DTId : SV_DispatchThreadID, uint3 GTId : SV_GroupThreadID, uint3 GId : SV_GroupID)\n{\n}\n");
+		NewIndex = Desc.Kernels.Num() - 1;
+		break;
+	}
+	case EComputeGraphItemKind::Interface:
+	{
+		TArray<FName> Existing;
+		for (FComputeGraphDataInterfaceDesc const& I : Desc.DataInterfaces)
+		{
+			Existing.Add(I.Name);
+		}
+
+		FComputeGraphDataInterfaceDesc& New = Desc.DataInterfaces.AddDefaulted_GetRef();
+		New.Name = GenerateUniqueName(Existing, TEXT("NewInterface"));
+		NewIndex = Desc.DataInterfaces.Num() - 1;
+		break;
+	}
+	case EComputeGraphItemKind::BindingObject:
+	{
+		TArray<FName> Existing;
+		for (FComputeGraphDataBindingObjectDesc const& B : Desc.BindingObjects)
+		{
+			Existing.Add(B.Name);
+		}
+
+		FComputeGraphDataBindingObjectDesc& New = Desc.BindingObjects.AddDefaulted_GetRef();
+		New.Name = GenerateUniqueName(Existing, TEXT("NewBindingObject"));
+		NewIndex = Desc.BindingObjects.Num() - 1;
+		break;
+	}
+	default:
+		return;
+	}
+
+	RefreshNavigatorAndDetails();
+	if (NavigatorWidget.IsValid())
+	{
+		NavigatorWidget->SetSelectedItem(Kind, NewIndex);
+	}
+}
+
+void FEditableComputeGraphEditorToolkit::OnDeleteItem(EComputeGraphItemKind Kind, int32 Index)
+{
+	if (Asset == nullptr)
+	{
+		return;
+	}
+
+	FScopedTransaction Transaction(LOCTEXT("DeleteItem", "Delete Compute Graph Item"));
+	Asset->Modify();
+
+	FComputeGraphDesc& Desc = Asset->GetGraphDescription();
+	switch (Kind)
+	{
+	case EComputeGraphItemKind::Kernel:
+		if (Desc.Kernels.IsValidIndex(Index)) 
+		{
+			Desc.Kernels.RemoveAt(Index);
+		}
+		break;
+	case EComputeGraphItemKind::Interface:
+		if (Desc.DataInterfaces.IsValidIndex(Index)) 
+		{
+			Desc.DataInterfaces.RemoveAt(Index);
+		}
+		break;
+	case EComputeGraphItemKind::BindingObject:
+		if (Desc.BindingObjects.IsValidIndex(Index)) 
+		{
+			Desc.BindingObjects.RemoveAt(Index);
+		}
+		break;
+	default:
+		break;
+	}
+
+	// Clear selection if we deleted the selected item.
+	if (EditorSelection && EditorSelection->Kind == Kind && EditorSelection->Index == Index)
+	{
+		EditorSelection->Kind = EComputeGraphItemKind::None;
+		EditorSelection->Index = INDEX_NONE;
+	}
+
+	RefreshNavigatorAndDetails();
+}
+
+void FEditableComputeGraphEditorToolkit::OnDuplicateItem(EComputeGraphItemKind Kind, int32 Index)
+{
+	if (Asset == nullptr) 
+	{
+		return;
+	}
+
+	FScopedTransaction Transaction(LOCTEXT("DuplicateItem", "Duplicate Compute Graph Item"));
+	Asset->Modify();
+
+	FComputeGraphDesc& Desc = Asset->GetGraphDescription();
+	int32 NewIndex = INDEX_NONE;
+
+	switch (Kind)
+	{
+	case EComputeGraphItemKind::Kernel:
+		if (Desc.Kernels.IsValidIndex(Index))
+		{
+			TArray<FName> Existing;
+			for (FComputeGraphKernelDesc const& Kernel : Desc.Kernels)
+			{
+				Existing.Add(Kernel.Name);
+			}
+
+			FComputeGraphKernelDesc Copy = Desc.Kernels[Index];
+			Copy.Name = GenerateUniqueName(Existing, Copy.Name.ToString() + TEXT("_Copy"));
+			Desc.Kernels.Add(MoveTemp(Copy));
+			NewIndex = Desc.Kernels.Num() - 1;
+		}
+		break;
+	case EComputeGraphItemKind::Interface:
+		if (Desc.DataInterfaces.IsValidIndex(Index))
+		{
+			TArray<FName> Existing;
+			for (FComputeGraphDataInterfaceDesc const& DataInterface : Desc.DataInterfaces)
+			{
+				Existing.Add(DataInterface.Name);
+			}
+
+			FComputeGraphDataInterfaceDesc Copy = Desc.DataInterfaces[Index];
+			Copy.Name = GenerateUniqueName(Existing, Copy.Name.ToString() + TEXT("_Copy"));
+			Desc.DataInterfaces.Add(MoveTemp(Copy));
+			NewIndex = Desc.DataInterfaces.Num() - 1;
+		}
+		break;
+	case EComputeGraphItemKind::BindingObject:
+		if (Desc.BindingObjects.IsValidIndex(Index))
+		{
+			TArray<FName> Existing;
+			for (FComputeGraphDataBindingObjectDesc const& BindingObject : Desc.BindingObjects)
+			{
+				Existing.Add(BindingObject.Name);
+			}
+
+			FComputeGraphDataBindingObjectDesc Copy = Desc.BindingObjects[Index];
+			Copy.Name = GenerateUniqueName(Existing, Copy.Name.ToString() + TEXT("_Copy"));
+			Desc.BindingObjects.Add(MoveTemp(Copy));
+			NewIndex = Desc.BindingObjects.Num() - 1;
+		}
+		break;
+	default:
+		break;
+	}
+
+	RefreshNavigatorAndDetails();
+	if (NavigatorWidget.IsValid() && NewIndex != INDEX_NONE)
+	{
+		NavigatorWidget->SetSelectedItem(Kind, NewIndex);
+	}
+}
+
+void FEditableComputeGraphEditorToolkit::OnRenameItem(EComputeGraphItemKind Kind, int32 Index, FName NewName)
+{
+	if (Asset == nullptr)
+	{
+		return;
+	}
+
+	FScopedTransaction Transaction(LOCTEXT("RenameItem", "Rename Compute Graph Item"));
+	Asset->Modify();
+
+	FComputeGraphDesc& Desc = Asset->GetGraphDescription();
+	switch (Kind)
+	{
+	case EComputeGraphItemKind::Kernel:
+		if (Desc.Kernels.IsValidIndex(Index))
+		{
+			Desc.Kernels[Index].Name = NewName;
+		}
+		break;
+	case EComputeGraphItemKind::Interface:
+		if (Desc.DataInterfaces.IsValidIndex(Index))
+		{
+			const FName OldName = Desc.DataInterfaces[Index].Name;
+			Desc.DataInterfaces[Index].Name = NewName;
+			RenameInterfaceInPins(Desc, OldName, NewName);
+			if (CachedSelectedItemName == OldName)
+			{
+				CachedSelectedItemName = NewName;
+			}
+		}
+		break;
+	case EComputeGraphItemKind::BindingObject:
+		if (Desc.BindingObjects.IsValidIndex(Index))
+		{
+			const FName OldName = Desc.BindingObjects[Index].Name;
+			Desc.BindingObjects[Index].Name = NewName;
+
+			RenameBindingObjectInInterfaces(Desc, OldName, NewName);
+			if (CachedSelectedItemName == OldName)
+			{
+				CachedSelectedItemName = NewName;
+			}
+		}
+		break;
+	default:
+		break;
+	}
+
+	RefreshNavigatorAndDetails();
+}
+
+void FEditableComputeGraphEditorToolkit::RefreshNavigatorAndDetails()
+{
+	bDirty = true;
+
+	if (NavigatorWidget.IsValid())
+	{
+		NavigatorWidget->Refresh();
+	}
+	if (GraphViewWidget.IsValid())
+	{
+		GraphViewWidget->Refresh();
+	}
+	DetailsView->ForceRefresh();
+}
+
+void FEditableComputeGraphEditorToolkit::OnHlslTextCommitted(FName KernelName, FString const& NewText)
+{
+	if (Asset == nullptr) 
+	{
+		return;
+	}
+
+	Asset->Modify();
+
+	FComputeGraphDesc& Desc = Asset->GetGraphDescription();
+	for (FComputeGraphKernelDesc& KernelDesc : Desc.Kernels)
+	{
+		if (KernelDesc.Name != KernelName) 
+		{
+			continue;
+		}
+
+		KernelDesc.SourceText = NewText;
+
+		TArray<FHlslExternalPinParser::FPinDeclaration> const Decls = FHlslExternalPinParser::FindExternalPins(NewText);
+
+		TSet<FString> DeclaredInputNames, DeclaredOutputNames;
+		for (FHlslExternalPinParser::FPinDeclaration const& Decl : Decls)
+		{
+			(Decl.bIsOutput ? DeclaredOutputNames : DeclaredInputNames).Add(Decl.FunctionName);
+		}
+
+		// Sync pins for one direction.
+		// Consider 1 new + 1 unmatched as a rename and update in place.
+		// Otherwise add new pins and mark unmatched ones as orphaned.
+		auto SyncPins = [](TArray<FKernelPin>& Pins, TSet<FString> const& DeclaredNames)
+		{
+			TArray<FString> ToAdd;
+			TArray<int32> UnmatchedIndices;
+			TSet<FString> ExistingNames;
+			
+			for (FKernelPin const& Pin : Pins)
+			{
+				ExistingNames.Add(Pin.KernelFunctionName);
+			}
+
+			for (FString const& Name : DeclaredNames)
+			{
+				if (!ExistingNames.Contains(Name))
+				{
+					ToAdd.Add(Name);
+				}
+			}
+
+			for (int32 i = 0; i < Pins.Num(); ++i)
+			{
+				if (!Pins[i].KernelFunctionName.IsEmpty() && !DeclaredNames.Contains(Pins[i].KernelFunctionName))
+				{
+					UnmatchedIndices.Add(i);
+				}
+			}
+
+			if (ToAdd.Num() == 1 && UnmatchedIndices.Num() == 1)
+			{
+				// Unambiguous rename so update the existing pin in place and preserve bindings.
+				Pins[UnmatchedIndices[0]].KernelFunctionName = ToAdd[0];
+				Pins[UnmatchedIndices[0]].bOrphaned = false;
+			}
+			else
+			{
+				for (FString const& Name : ToAdd)
+				{
+					Pins.AddDefaulted_GetRef().KernelFunctionName = Name;
+				}
+				for (int32 Idx : UnmatchedIndices)
+				{
+					Pins[Idx].bOrphaned = true;
+				}
+			}
+		};
+
+		SyncPins(KernelDesc.Inputs, DeclaredInputNames);
+		SyncPins(KernelDesc.Outputs, DeclaredOutputNames);
+
+		// Sync EntryPoint.
+		// If it no longer appears as a function in the HLSL but exactly one non-DI_ void function does then treat as a rename.
+		if (!KernelDesc.EntryPoint.IsEmpty() && !FHlslExternalPinParser::FunctionExistsInText(NewText, KernelDesc.EntryPoint))
+		{
+			TArray<FString> Candidates = FHlslExternalPinParser::FindFunctionDefinitions(NewText);
+			Candidates.RemoveAll([](FString const& Name) { return Name.StartsWith(TEXT("DI_")); });
+			if (Candidates.Num() == 1)
+			{
+				KernelDesc.EntryPoint = Candidates[0];
+			}
+		}
+
+		bDirty = true;
+		break;
+	}
+
+	if (HlslEditorWidget.IsValid())
+	{
+		RefreshHlslBoundFunctions(KernelName);
+	}
+	if (NavigatorWidget.IsValid())
+	{
+		NavigatorWidget->Refresh();
+	}
+	if (GraphViewWidget.IsValid())
+	{
+		GraphViewWidget->Refresh();
+	}
+	DetailsView->ForceRefresh();
+}
+
+void FEditableComputeGraphEditorToolkit::CompileGraph()
+{
+	if (Asset == nullptr) 
+	{
+		return;
+	}
+
+	// Flush any uncommitted HLSL text so the compile sees the latest source.
+	if (HlslEditorWidget.IsValid())
+	{
+		HlslEditorWidget->CommitCurrentText();
+	}
+
+	Asset->RebuildGraph();
+
+	bDirty = false;
+}
+
+#undef LOCTEXT_NAMESPACE
diff --git a/Engine/Plugins/Runtime/ComputeFramework/Source/EditableComputeGraphEditor/Private/EditableComputeGraphFactory.cpp b/Engine/Plugins/Runtime/ComputeFramework/Source/EditableComputeGraphEditor/Private/EditableComputeGraphFactory.cpp
new file mode 100644
--- /dev/null
+++ b/Engine/Plugins/Runtime/ComputeFramework/Source/EditableComputeGraphEditor/Private/EditableComputeGraphFactory.cpp
@@ -0,0 +1,19 @@
+// Copyright Epic Games, Inc. All Rights Reserved.
+
+#include "ComputeFramework/EditableComputeGraphFactory.h"
+
+#include "ComputeFramework/EditableComputeGraph.h"
+
+#include UE_INLINE_GENERATED_CPP_BY_NAME(EditableComputeGraphFactory)
+
+UEditableComputeGraphFactory::UEditableComputeGraphFactory()
+{
+	SupportedClass = UEditableComputeGraph::StaticClass();
+	bCreateNew = true;
+	bEditAfterNew = true;
+}
+
+UObject* UEditableComputeGraphFactory::FactoryCreateNew(UClass* InClass, UObject* InParent, FName InName, EObjectFlags Flags, UObject* Context, FFeedbackContext* Warn)
+{
+	return NewObject<UEditableComputeGraph>(InParent, InClass, InName, Flags);
+}
diff --git a/Engine/Plugins/Runtime/ComputeFramework/Source/EditableComputeGraphEditor/Private/HlslExternalPinParser.cpp b/Engine/Plugins/Runtime/ComputeFramework/Source/EditableComputeGraphEditor/Private/HlslExternalPinParser.cpp
new file mode 100644
--- /dev/null
+++ b/Engine/Plugins/Runtime/ComputeFramework/Source/EditableComputeGraphEditor/Private/HlslExternalPinParser.cpp
@@ -0,0 +1,303 @@
+// Copyright Epic Games, Inc. All Rights Reserved.
+
+#include "ComputeFramework/HlslExternalPinParser.h"
+
+#include "Containers/Set.h"
+#include "Containers/StringView.h"
+
+namespace HlslExternalPinParser {
+
+/** Replace all content inside line and block comments with spaces, preserving newlines so that line numbers are unchanged. */
+static FString StripComments(FStringView Src)
+{
+	FString Out;
+	Out.Reserve(Src.Len());
+
+	const int32 Len = Src.Len();
+	int32 Pos = 0;
+	while (Pos < Len)
+	{
+		if (Pos + 1 < Len && Src[Pos] == TEXT('/') && Src[Pos + 1] == TEXT('/'))
+		{
+			// Line comment: replace up to (but not including) the newline.
+			Pos += 2;
+			while (Pos < Len && Src[Pos] != TEXT('\n'))
+			{
+				Out += TEXT(' ');
+				++Pos;
+			}
+		}
+		else if (Pos + 1 < Len && Src[Pos] == TEXT('/') && Src[Pos + 1] == TEXT('*'))
+		{
+			// Block comment: replace content; preserve newlines.
+			Pos += 2;
+			Out += TEXT(' ');
+			Out += TEXT(' ');
+			while (Pos + 1 < Len && !(Src[Pos] == TEXT('*') && Src[Pos + 1] == TEXT('/')))
+			{
+				Out += (Src[Pos] == TEXT('\n')) ? TEXT('\n') : TEXT(' ');
+				++Pos;
+			}
+			if (Pos + 1 < Len)
+			{
+				Out += TEXT(' ');
+				Out += TEXT(' ');
+				Pos += 2;
+			}
+		}
+		else
+		{
+			Out += Src[Pos++];
+		}
+	}
+	return Out;
+}
+
+} // namespace HlslExternalPinParser
+
+TArray<FHlslExternalPinParser::FPinDeclaration> FHlslExternalPinParser::FindExternalPins(FStringView SourceText)
+{
+	TArray<FHlslExternalPinParser::FPinDeclaration> Result;
+	TSet<FString> Seen;
+
+	const FString Stripped = HlslExternalPinParser::StripComments(SourceText);
+	TCHAR const* Data = *Stripped;
+	const int32 Len = Stripped.Len();
+
+	// Scan for "DI_" followed by an identifier character.
+	int32 Pos = 0;
+	while (Pos < Len)
+	{
+		if (Pos + 3 < Len && Data[Pos] == TEXT('D') && Data[Pos + 1] == TEXT('I') && Data[Pos + 2] == TEXT('_') && (FChar::IsAlpha(Data[Pos + 3]) || Data[Pos + 3] == TEXT('_')))
+		{
+			// Ensure DI_ is not the tail of a longer identifier (e.g. FOO_DI_Bar).
+			const bool bValidStart = (Pos == 0)	|| !TChar<TCHAR>::IsIdentifier(Data[Pos - 1]);
+
+			if (bValidStart)
+			{
+				// Collect the full identifier starting at 'D'.
+				const int32 NameStart = Pos;
+				Pos += 3; // skip "DI_"
+				while (Pos < Len && TChar<TCHAR>::IsIdentifier(Data[Pos]))
+				{
+					++Pos;
+				}
+				const FString FuncName(Pos - NameStart, Data + NameStart);
+
+				// Must be followed by '(' (allowing whitespace between).
+				int32 AfterNamePos = Pos;
+				while (AfterNamePos < Len && (Data[AfterNamePos] == TEXT(' ') || Data[AfterNamePos] == TEXT('\t')))
+				{
+					++AfterNamePos;
+				}
+
+				if (AfterNamePos < Len && Data[AfterNamePos] == TEXT('(') && !Seen.Contains(FuncName))
+				{
+					Seen.Add(FuncName);
+
+					// Infer direction by scanning backward to the start of the statement. A "statement start" is a newline, '{', '}', or ';'.
+					// If only whitespace (spaces/tabs/carriage-returns) precede DI_ on the current line, the call is a standalone statement → void → output pin.
+					// If anything else appears before it, the result is used -> input pin.
+					bool bStandaloneStatement = true;
+					for (int32 ScanBack = NameStart - 1; ScanBack >= 0; --ScanBack)
+					{
+						const TCHAR PrecedingChar = Data[ScanBack];
+						if (PrecedingChar == TEXT('\n') || PrecedingChar == TEXT('{') || PrecedingChar == TEXT('}') || PrecedingChar == TEXT(';'))
+						{
+							break;
+						}
+						if (PrecedingChar != TEXT(' ') && PrecedingChar != TEXT('\t') && PrecedingChar != TEXT('\r'))
+						{
+							bStandaloneStatement = false;
+							break;
+						}
+					}
+
+					Result.Add({ FuncName, /*bIsOutput=*/bStandaloneStatement });
+				}
+				continue; // Pos already advanced past the identifier
+			}
+		}
+		++Pos;
+	}
+
+	return Result;
+}
+
+TArray<FString> FHlslExternalPinParser::FindFunctionDefinitions(FStringView SourceText)
+{
+	TArray<FString> Result;
+
+	const FString Stripped = HlslExternalPinParser::StripComments(SourceText);
+	TCHAR const* Data = *Stripped;
+	const int32 Len = Stripped.Len();
+
+	int32 BraceDepth = 0;
+	int32 Pos = 0;
+	TArray<FString> Words;
+
+	auto SkipAllWhitespace = [&]()
+	{
+		while (Pos < Len && (Data[Pos] == TEXT(' ') || Data[Pos] == TEXT('\t') || Data[Pos] == TEXT('\r') || Data[Pos] == TEXT('\n')))
+		{
+			++Pos;
+		}
+	};
+
+	while (Pos < Len)
+	{
+		const TCHAR CurrentChar = Data[Pos];
+
+		if (CurrentChar == TEXT('{'))
+		{
+			++BraceDepth;
+			++Pos;
+			continue;
+		}
+		if (CurrentChar == TEXT('}'))
+		{
+			if (BraceDepth > 0)
+			{
+				--BraceDepth;
+			}
+			++Pos;
+			continue;
+		}
+
+		// Only scan at global scope.
+		if (BraceDepth > 0)
+		{
+			++Pos;
+			continue;
+		}
+
+		// Skip whitespace, semicolons, and attribute blocks.
+		if (CurrentChar == TEXT(' ') || CurrentChar == TEXT('\t') || CurrentChar == TEXT('\r') || CurrentChar == TEXT('\n') || CurrentChar == TEXT(';'))
+		{
+			++Pos;
+			continue;
+		}
+		if (CurrentChar == TEXT('['))
+		{
+			++Pos;
+			while (Pos < Len && Data[Pos] != TEXT(']'))
+			{
+				++Pos;
+			}
+			if (Pos < Len)
+			{
+				++Pos;
+			}
+			continue;
+		}
+
+		if (FChar::IsAlpha(CurrentChar) || CurrentChar == TEXT('_'))
+		{
+			// Collect whitespace-separated identifier tokens up to '('.
+			Words.Reset();
+			while (Pos < Len)
+			{
+				while (Pos < Len && (Data[Pos] == TEXT(' ') || Data[Pos] == TEXT('\t')))
+				{
+					++Pos;
+				}
+				if (Pos >= Len)
+				{
+					break;
+				}
+				if (Data[Pos] == TEXT('('))
+				{
+					break;
+				}
+				if (FChar::IsAlpha(Data[Pos]) || Data[Pos] == TEXT('_'))
+				{
+					const int32 WordStart = Pos;
+					while (Pos < Len && TChar<TCHAR>::IsIdentifier(Data[Pos]))
+					{
+						++Pos;
+					}
+					Words.Add(FString(Pos - WordStart, Data + WordStart));
+				}
+				else
+				{
+					Words.Reset();
+					break;
+				}
+			}
+
+			if (Words.Num() < 2 || Pos >= Len || Data[Pos] != TEXT('('))
+			{
+				continue;
+			}
+
+			const FString FuncName = Words.Last();
+
+			// Skip the parameter list.
+			int32 ParenDepth = 1;
+			++Pos;
+			while (Pos < Len && ParenDepth > 0)
+			{
+				if (Data[Pos] == TEXT('('))
+				{
+					++ParenDepth;
+				}
+				else if (Data[Pos] == TEXT(')'))
+				{
+					--ParenDepth;
+				}
+				++Pos;
+			}
+
+			// Skip optional semantic (': SV_Something') and whitespace.
+			SkipAllWhitespace();
+			if (Pos < Len && Data[Pos] == TEXT(':'))
+			{
+				while (Pos < Len && Data[Pos] != TEXT('{') && Data[Pos] != TEXT(';'))
+				{
+					++Pos;
+				}
+			}
+			SkipAllWhitespace();
+
+			if (Pos < Len && Data[Pos] == TEXT('{'))
+			{
+				Result.Add(FuncName);
+				// Leave '{' for the brace-depth handler at the top of the loop.
+			}
+			continue;
+		}
+
+		++Pos;
+	}
+
+	return Result;
+}
+
+bool FHlslExternalPinParser::FunctionExistsInText(FStringView SourceText, FString const& FunctionName)
+{
+	const FString Stripped = HlslExternalPinParser::StripComments(SourceText);
+	const int32 Len = Stripped.Len();
+	int32 SearchFrom = 0;
+	while (true)
+	{
+		const int32 Found = Stripped.Find(FunctionName, ESearchCase::CaseSensitive, ESearchDir::FromStart, SearchFrom);
+		if (Found == INDEX_NONE)
+		{
+			return false;
+		}
+
+		const bool bPrecededOk = (Found == 0) || !TChar<TCHAR>::IsIdentifier(Stripped[Found - 1]);
+		int32 AfterMatchPos = Found + FunctionName.Len();
+
+		while (AfterMatchPos < Len && FChar::IsWhitespace(Stripped[AfterMatchPos]))
+		{
+			++AfterMatchPos;
+		}
+		if (bPrecededOk && AfterMatchPos < Len && Stripped[AfterMatchPos] == TEXT('('))
+		{
+			return true;
+		}
+
+		SearchFrom = Found + 1;
+	}
+}
diff --git a/Engine/Plugins/Runtime/ComputeFramework/Source/EditableComputeGraphEditor/Private/KernelPinCustomization.cpp b/Engine/Plugins/Runtime/ComputeFramework/Source/EditableComputeGraphEditor/Private/KernelPinCustomization.cpp
new file mode 100644
--- /dev/null
+++ b/Engine/Plugins/Runtime/ComputeFramework/Source/EditableComputeGraphEditor/Private/KernelPinCustomization.cpp
@@ -0,0 +1,415 @@
+// Copyright Epic Games, Inc. All Rights Reserved.
+
+#include "ComputeFramework/KernelPinCustomization.h"
+
+#include "ComputeFramework/ComputeDataInterface.h"
+#include "ComputeFramework/EditableComputeGraph.h"
+#include "ComputeFramework/HlslExternalPinParser.h"
+#include "ComputeFramework/ShaderParamTypeDefinition.h"
+#include "DetailLayoutBuilder.h"
+#include "DetailWidgetRow.h"
+#include "IDetailChildrenBuilder.h"
+#include "PropertyHandle.h"
+#include "Styling/AppStyle.h"
+#include "Widgets/Input/SComboBox.h"
+#include "Widgets/Text/STextBlock.h"
+
+#define LOCTEXT_NAMESPACE "ComputeFrameworkEditor"
+
+TSharedRef<IPropertyTypeCustomization> FKernelPinCustomization::MakeInstance()
+{
+	return MakeShared<FKernelPinCustomization>();
+}
+
+void FKernelPinCustomization::CustomizeHeader(TSharedRef<IPropertyHandle> StructPropertyHandle,	FDetailWidgetRow& HeaderRow, IPropertyTypeCustomizationUtils& CustomizationUtils)
+{
+	StructHandle = StructPropertyHandle;
+	KernelFnHandle = StructPropertyHandle->GetChildHandle(GET_MEMBER_NAME_CHECKED(FKernelPin, KernelFunctionName));
+	InterfaceNameHandle = StructPropertyHandle->GetChildHandle(GET_MEMBER_NAME_CHECKED(FKernelPin, DataInterfaceName));
+	FunctionNameHandle = StructPropertyHandle->GetChildHandle(GET_MEMBER_NAME_CHECKED(FKernelPin, DataInterfaceFunctionName));
+
+	// Determine direction from the name of the containing array property.
+	TSharedPtr<IPropertyHandle> ArrayHandle = StructPropertyHandle->GetParentHandle();
+	bIsInput = ArrayHandle.IsValid() && ArrayHandle->GetProperty() && ArrayHandle->GetProperty()->GetFName() == GET_MEMBER_NAME_CHECKED(FComputeGraphKernelDesc, Inputs);
+
+	// Navigate up one more level to the FComputeGraphKernelDesc element for source text access.
+	TSharedPtr<IPropertyHandle> KernelDescHandle = ArrayHandle.IsValid() ? ArrayHandle->GetParentHandle() : nullptr;
+	if (KernelDescHandle.IsValid())
+	{
+		KernelSourceHandle = KernelDescHandle->GetChildHandle(GET_MEMBER_NAME_CHECKED(FComputeGraphKernelDesc, SourceText));
+		KernelEntryPointHandle = KernelDescHandle->GetChildHandle(GET_MEMBER_NAME_CHECKED(FComputeGraphKernelDesc, EntryPoint));
+	}
+
+	OrphanedHandle = StructPropertyHandle->GetChildHandle(GET_MEMBER_NAME_CHECKED(FKernelPin, bOrphaned));
+
+	HeaderRow
+		.NameContent()
+		[
+			StructPropertyHandle->CreatePropertyNameWidget()
+		]
+		.ValueContent()
+		.MinDesiredWidth(300.f)
+		[
+			SNew(STextBlock)
+			.Text(this, &FKernelPinCustomization::GetHeaderSummaryText)
+			.Font(CustomizationUtils.GetRegularFont())
+			.ColorAndOpacity(this, &FKernelPinCustomization::GetHeaderSummaryColor)
+			.ToolTipText(this, &FKernelPinCustomization::GetOrphanedTooltipText)
+		];
+}
+
+void FKernelPinCustomization::CustomizeChildren(TSharedRef<IPropertyHandle> StructPropertyHandle, IDetailChildrenBuilder& ChildBuilder, IPropertyTypeCustomizationUtils& CustomizationUtils)
+{
+	// The KernelFunctionName combo box is populated from parsed HLSL pin declarations.
+	ChildBuilder.AddCustomRow(LOCTEXT("KernelFunctionName", "Kernel Function"))
+		.NameContent()
+		[
+			KernelFnHandle->CreatePropertyNameWidget()
+		]
+		.ValueContent()
+		.MinDesiredWidth(200.f)
+		[
+			SNew(SComboBox<TSharedPtr<FString>>)
+			.OptionsSource(&KernelFunctionNameOptions)
+			.OnComboBoxOpening(this, &FKernelPinCustomization::RefreshKernelFunctionNameOptions)
+			.OnSelectionChanged(this, &FKernelPinCustomization::OnKernelFunctionSelected)
+			.OnGenerateWidget_Lambda([](TSharedPtr<FString> Item) -> TSharedRef<SWidget>
+			{
+				return SNew(STextBlock)
+					.Text(FText::FromString(Item.IsValid() ? *Item : FString()))
+					.Font(FAppStyle::GetFontStyle("PropertyWindow.NormalFont"));
+			})
+			[
+				SNew(STextBlock)
+				.Text(this, &FKernelPinCustomization::GetCurrentKernelFunctionNameText)
+				.Font(FAppStyle::GetFontStyle("PropertyWindow.NormalFont"))
+			]
+		];
+
+	// DataInterfaceName combo box.
+	ChildBuilder.AddCustomRow(LOCTEXT("DataInterfaceName", "Data Interface"))
+		.NameContent()
+		[
+			InterfaceNameHandle->CreatePropertyNameWidget()
+		]
+		.ValueContent()
+		.MinDesiredWidth(200.f)
+		[
+			SNew(SComboBox<TSharedPtr<FName>>)
+			.OptionsSource(&InterfaceNameOptions)
+			.OnComboBoxOpening(this, &FKernelPinCustomization::RefreshInterfaceNameOptions)
+			.OnSelectionChanged(this, &FKernelPinCustomization::OnInterfaceSelected)
+			.OnGenerateWidget_Lambda([](TSharedPtr<FName> Item) -> TSharedRef<SWidget>
+			{
+				return SNew(STextBlock)
+					.Text(FText::FromName(Item.IsValid() ? *Item : NAME_None))
+					.Font(FAppStyle::GetFontStyle("PropertyWindow.NormalFont"));
+			})
+			[
+				SNew(STextBlock)
+				.Text(this, &FKernelPinCustomization::GetCurrentInterfaceNameText)
+				.Font(FAppStyle::GetFontStyle("PropertyWindow.NormalFont"))
+			]
+		];
+
+	// DataInterfaceFunctionName combo box.
+	ChildBuilder.AddCustomRow(LOCTEXT("DataInterfaceFunctionName", "Function"))
+		.NameContent()
+		[
+			FunctionNameHandle->CreatePropertyNameWidget()
+		]
+		.ValueContent()
+		.MinDesiredWidth(200.f)
+		[
+			SNew(SComboBox<TSharedPtr<FString>>)
+			.OptionsSource(&FunctionNameOptions)
+			.OnComboBoxOpening(this, &FKernelPinCustomization::RefreshFunctionNameOptions)
+			.OnSelectionChanged(this, &FKernelPinCustomization::OnFunctionSelected)
+			.OnGenerateWidget_Lambda([](TSharedPtr<FString> Item) -> TSharedRef<SWidget>
+			{
+				return SNew(STextBlock)
+					.Text(FText::FromString(Item.IsValid() ? *Item : FString()))
+					.Font(FAppStyle::GetFontStyle("PropertyWindow.NormalFont"));
+			})
+			[
+				SNew(STextBlock)
+				.Text(this, &FKernelPinCustomization::GetCurrentFunctionNameText)
+				.Font(FAppStyle::GetFontStyle("PropertyWindow.NormalFont"))
+			]
+		];
+}
+
+void FKernelPinCustomization::RefreshKernelFunctionNameOptions()
+{
+	KernelFunctionNameOptions.Reset();
+
+	if (!KernelSourceHandle.IsValid() || !KernelEntryPointHandle.IsValid())
+	{
+		return;
+	}
+
+	FString SourceText, EntryPoint;
+	if (KernelSourceHandle->GetValue(SourceText) != FPropertyAccess::Success || SourceText.IsEmpty())
+	{
+		return;
+	}
+	KernelEntryPointHandle->GetValue(EntryPoint);
+
+	const TArray<FHlslExternalPinParser::FPinDeclaration> Decls = FHlslExternalPinParser::FindExternalPins(SourceText);
+	for (FHlslExternalPinParser::FPinDeclaration const& Decl : Decls)
+	{
+		// Input pins read from the data interface (non-void return).
+		// Output pins write (void return).
+		const bool bWantOutput = !bIsInput;
+		if (Decl.bIsOutput == bWantOutput)
+		{
+			KernelFunctionNameOptions.Add(MakeShared<FString>(Decl.FunctionName));
+		}
+	}
+}
+
+void FKernelPinCustomization::RefreshInterfaceNameOptions()
+{
+	InterfaceNameOptions.Reset();
+
+	if (!StructHandle.IsValid())
+	{
+		return;
+	}
+
+	TArray<UObject*> OuterObjects;
+	StructHandle->GetOuterObjects(OuterObjects);
+	if (OuterObjects.Num() != 1)
+	{
+		return;
+	}
+
+	UEditableComputeGraph const* Graph = Cast<UEditableComputeGraph>(OuterObjects[0]);
+	if (!Graph)
+	{
+		return;
+	}
+
+	for (FComputeGraphDataInterfaceDesc const& InterfaceDesc : Graph->GetGraphDescription().DataInterfaces)
+	{
+		if (!InterfaceDesc.Name.IsNone())
+		{
+			InterfaceNameOptions.Add(MakeShared<FName>(InterfaceDesc.Name));
+		}
+	}
+}
+
+void FKernelPinCustomization::RefreshFunctionNameOptions()
+{
+	FunctionNameOptions.Reset();
+
+	if (!StructHandle.IsValid() || !InterfaceNameHandle.IsValid())
+	{
+		return;
+	}
+
+	FName CurrentInterfaceName;
+	if (InterfaceNameHandle->GetValue(CurrentInterfaceName) != FPropertyAccess::Success || CurrentInterfaceName.IsNone())
+	{
+		return;
+	}
+
+	TArray<UObject*> OuterObjects;
+	StructHandle->GetOuterObjects(OuterObjects);
+	if (OuterObjects.Num() != 1)
+	{
+		return;
+	}
+
+	UEditableComputeGraph const* Graph = Cast<UEditableComputeGraph>(OuterObjects[0]);
+	if (!Graph)
+	{
+		return;
+	}
+
+	for (FComputeGraphDataInterfaceDesc const& InterfaceDesc : Graph->GetGraphDescription().DataInterfaces)
+	{
+		if (InterfaceDesc.Name != CurrentInterfaceName)
+		{
+			continue;
+		}
+
+		// Prefer any configured Settings, but fall back to the CDO.
+		UComputeDataInterface const* Interface = InterfaceDesc.Settings
+			? InterfaceDesc.Settings.Get()
+			: (InterfaceDesc.Type ? Cast<UComputeDataInterface>(InterfaceDesc.Type->GetDefaultObject()) : nullptr);
+
+		if (!Interface)
+		{
+			break;
+		}
+
+		TArray<FShaderFunctionDefinition> Functions;
+		if (bIsInput)
+		{
+			Interface->GetSupportedInputs(Functions);
+		}
+		else
+		{
+			Interface->GetSupportedOutputs(Functions);
+		}
+
+		for (FShaderFunctionDefinition const& Fn : Functions)
+		{
+			FunctionNameOptions.Add(MakeShared<FString>(Fn.Name));
+		}
+		break;
+	}
+}
+
+void FKernelPinCustomization::OnKernelFunctionSelected(TSharedPtr<FString> NewValue, ESelectInfo::Type SelectInfo)
+{
+	if (KernelFnHandle.IsValid() && NewValue.IsValid())
+	{
+		KernelFnHandle->SetValue(*NewValue);
+	}
+}
+
+void FKernelPinCustomization::OnInterfaceSelected(TSharedPtr<FName> NewValue, ESelectInfo::Type SelectInfo)
+{
+	if (InterfaceNameHandle.IsValid() && NewValue.IsValid())
+	{
+		InterfaceNameHandle->SetValue(*NewValue);
+		
+		// Clear the function name. It belonged to the old interface.
+		if (FunctionNameHandle.IsValid())
+		{
+			FunctionNameHandle->SetValue(FString());
+		}
+	}
+}
+
+void FKernelPinCustomization::OnFunctionSelected(TSharedPtr<FString> NewValue, ESelectInfo::Type SelectInfo)
+{
+	if (FunctionNameHandle.IsValid() && NewValue.IsValid())
+	{
+		FunctionNameHandle->SetValue(*NewValue);
+	}
+}
+
+FSlateColor FKernelPinCustomization::GetHeaderSummaryColor() const
+{
+	return NeedsAttention() ? FSlateColor(FLinearColor(1.f, 0.4f, 0.2f, 1.f)) : FSlateColor::UseForeground();
+}
+
+FText FKernelPinCustomization::GetOrphanedTooltipText() const
+{
+	if (!NeedsAttention())
+	{
+		return FText::GetEmpty();
+	}
+
+	bool bOrphaned = false;
+	if (OrphanedHandle.IsValid())
+	{
+		OrphanedHandle->GetValue(bOrphaned);
+	}
+
+	return bOrphaned
+		? LOCTEXT("OrphanedPinTooltip", "This pin's kernel function no longer appears in the HLSL. Update or remove the binding.")
+		: LOCTEXT("IncompletePinTooltip", "This pin is not fully configured. Set the kernel function, data interface, and data interface function.");
+}
+
+bool FKernelPinCustomization::NeedsAttention() const
+{
+	FString KernelFn;
+	if (KernelFnHandle.IsValid() && KernelFnHandle->GetValue(KernelFn) == FPropertyAccess::Success && KernelFn.IsEmpty())
+	{
+		return true;
+	}
+
+	FName InterfaceName;
+	if (InterfaceNameHandle.IsValid() && InterfaceNameHandle->GetValue(InterfaceName) == FPropertyAccess::Success && InterfaceName.IsNone())
+	{
+		return true;
+	}
+
+	FString FunctionName;
+	if (FunctionNameHandle.IsValid() && FunctionNameHandle->GetValue(FunctionName) == FPropertyAccess::Success && FunctionName.IsEmpty())
+	{
+		return true;
+	}
+
+	bool bOrphaned = false;
+	if (OrphanedHandle.IsValid()) 
+	{
+		OrphanedHandle->GetValue(bOrphaned);
+	}
+	return bOrphaned;
+}
+
+FText FKernelPinCustomization::GetHeaderSummaryText() const
+{
+	if (!KernelFnHandle.IsValid())
+	{
+		return FText::GetEmpty();
+	}
+
+	FString KernelFn;
+	KernelFnHandle->GetValue(KernelFn);
+	
+	FName InterfaceName;
+	InterfaceNameHandle->GetValue(InterfaceName);
+	
+	FString InterfaceFn;
+	FunctionNameHandle->GetValue(InterfaceFn);
+
+	if (KernelFn.IsEmpty() && InterfaceName.IsNone() && InterfaceFn.IsEmpty())
+	{
+		return LOCTEXT("EmptyPin", "(empty)");
+	}
+
+	const FString Arrow = bIsInput ? TEXT(" \u2190 ") : TEXT(" \u2192 ");
+	FString Summary = KernelFn + Arrow + InterfaceName.ToString();
+	if (!InterfaceFn.IsEmpty())
+	{
+		Summary += TEXT(".") + InterfaceFn;
+	}
+	return FText::FromString(Summary);
+}
+
+FText FKernelPinCustomization::GetCurrentKernelFunctionNameText() const
+{
+	FString Value;
+	if (KernelFnHandle.IsValid() && KernelFnHandle->GetValue(Value) == FPropertyAccess::Success)
+	{
+		if (!Value.IsEmpty())
+		{
+			return FText::FromString(Value);
+		}
+	}
+	return LOCTEXT("SelectKernelFn", "(select)");
+}
+
+FText FKernelPinCustomization::GetCurrentInterfaceNameText() const
+{
+	FName Value;
+	if (InterfaceNameHandle.IsValid() && InterfaceNameHandle->GetValue(Value) == FPropertyAccess::Success)
+	{
+		if (!Value.IsNone())
+		{
+			return FText::FromName(Value);
+		}
+	}
+	return LOCTEXT("SelectInterface", "(select)");
+}
+
+FText FKernelPinCustomization::GetCurrentFunctionNameText() const
+{
+	FString Value;
+	if (FunctionNameHandle.IsValid() && FunctionNameHandle->GetValue(Value) == FPropertyAccess::Success)
+	{
+		if (!Value.IsEmpty())
+		{
+			return FText::FromString(Value);
+		}
+	}
+	return LOCTEXT("SelectFunction", "(select)");
+}
+
+#undef LOCTEXT_NAMESPACE
