Fix ASan error on shutdown.
Use FName directly instead of fetching from GC'd class.
[RN] rnx
#Synced-CL 52245232

diff --git a/Engine/Plugins/Runtime/ComputeFramework/Source/EditableComputeGraphEditor/Private/EditableComputeGraphEditorModule.cpp b/Engine/Plugins/Runtime/ComputeFramework/Source/EditableComputeGraphEditor/Private/EditableComputeGraphEditorModule.cpp
--- a/Engine/Plugins/Runtime/ComputeFramework/Source/EditableComputeGraphEditor/Private/EditableComputeGraphEditorModule.cpp
+++ b/Engine/Plugins/Runtime/ComputeFramework/Source/EditableComputeGraphEditor/Private/EditableComputeGraphEditorModule.cpp
@@ -11,7 +11,7 @@ public:
 	void StartupModule() override
 	{
 		FPropertyEditorModule& PropertyModule = FModuleManager::LoadModuleChecked<FPropertyEditorModule>("PropertyEditor");
-		PropertyModule.RegisterCustomPropertyTypeLayout(FKernelPin::StaticStruct()->GetFName(),	FOnGetPropertyTypeCustomizationInstance::CreateStatic(&FKernelPinCustomization::MakeInstance));
+		PropertyModule.RegisterCustomPropertyTypeLayout("KernelPin", FOnGetPropertyTypeCustomizationInstance::CreateStatic(&FKernelPinCustomization::MakeInstance));
 	}
 
 	void ShutdownModule() override
@@ -19,7 +19,7 @@ public:
 		if (FModuleManager::Get().IsModuleLoaded("PropertyEditor"))
 		{
 			FPropertyEditorModule& PropertyModule = FModuleManager::GetModuleChecked<FPropertyEditorModule>("PropertyEditor");
-			PropertyModule.UnregisterCustomPropertyTypeLayout(FKernelPin::StaticStruct()->GetFName());
+			PropertyModule.UnregisterCustomPropertyTypeLayout("KernelPin");
 		}
 	}
 };
