ComputeFrameork: Fix CIS error.
New file has code behine WITH_EDITOR but that doesn't match code in a merge integration of the header change.
#rnx
#Synced-CL 52350961

diff --git a/Engine/Plugins/Runtime/ComputeFramework/Source/ComputeFramework/Private/ComputeKernelSource.cpp b/Engine/Plugins/Runtime/ComputeFramework/Source/ComputeFramework/Private/ComputeKernelSource.cpp
--- a/Engine/Plugins/Runtime/ComputeFramework/Source/ComputeFramework/Private/ComputeKernelSource.cpp
+++ b/Engine/Plugins/Runtime/ComputeFramework/Source/ComputeFramework/Private/ComputeKernelSource.cpp
@@ -4,8 +4,6 @@
 
 #include "Internationalization/Regex.h"
 
-#if WITH_EDITOR
-
 /** Strip regex escape characters that can be in EntryPoint since it is a user editable string. */
 static FString StripRegexCharacters(const FString& Input)
 {
@@ -70,5 +68,3 @@ FString UComputeKernelSourceWithText::GetSource() const
 
 	return SourceText;
 }
-
-#endif // WITH_EDITOR
