ComputeFramework: Move some code to WITH_EDITOR in anticipation of incoming merge that moves data to WITH_EDITORONLY_DATA.
#rnx
#Synced-CL 52465805

diff --git a/Engine/Plugins/Animation/DeformerGraph/Source/OptimusCore/Public/OptimusKernelSource.h b/Engine/Plugins/Animation/DeformerGraph/Source/OptimusCore/Public/OptimusKernelSource.h
--- a/Engine/Plugins/Animation/DeformerGraph/Source/OptimusCore/Public/OptimusKernelSource.h
+++ b/Engine/Plugins/Animation/DeformerGraph/Source/OptimusCore/Public/OptimusKernelSource.h
@@ -17,12 +17,14 @@ public:
 		Source = InSource;
 	}
 	
+#if WITH_EDITOR
 	//~ Begin UComputeKernelSource Interface.
 	FString GetSource() const override
 	{
 		return Source;
 	}
 	//~ End UComputeKernelSource Interface.
+#endif
 
 protected:
 	UPROPERTY()
diff --git a/Engine/Plugins/PCG/Source/PCG/Private/Compute/PCGComputeKernelSource.h b/Engine/Plugins/PCG/Source/PCG/Private/Compute/PCGComputeKernelSource.h
--- a/Engine/Plugins/PCG/Source/PCG/Private/Compute/PCGComputeKernelSource.h
+++ b/Engine/Plugins/PCG/Source/PCG/Private/Compute/PCGComputeKernelSource.h
@@ -12,20 +12,18 @@ class UPCGComputeKernelSource : public UComputeKernelSource
 	GENERATED_BODY()
 
 public:
+#if WITH_EDITOR
 	//~ Begin UComputeKernelSource Interface.
 	FString GetSource() const override
 	{
-#if WITH_EDITOR
 		return Source;
-#else
-		ensure(false);
-		return "";
-#endif
 	}
 	//~ End UComputeKernelSource Interface.
 
-#if WITH_EDITOR
-	void SetSource(FString const& InSource) { Source = InSource; }
+	void SetSource(FString const& InSource)
+	{
+		Source = InSource;
+	}
 
 protected:
 	FString Source;
diff --git a/Engine/Plugins/Runtime/ComputeFramework/Source/ComputeFramework/Private/ComputeKernelSource.cpp b/Engine/Plugins/Runtime/ComputeFramework/Source/ComputeFramework/Private/ComputeKernelSource.cpp
--- a/Engine/Plugins/Runtime/ComputeFramework/Source/ComputeFramework/Private/ComputeKernelSource.cpp
+++ b/Engine/Plugins/Runtime/ComputeFramework/Source/ComputeFramework/Private/ComputeKernelSource.cpp
@@ -2,6 +2,8 @@
 
 #include "ComputeFramework/ComputeKernelSource.h"
 
+#if WITH_EDITOR
+
 #include "Internationalization/Regex.h"
 
 /** Strip regex escape characters that can be in EntryPoint since it is a user editable string. */
@@ -68,3 +70,5 @@ FString UComputeKernelSourceWithText::GetSource() const
 
 	return SourceText;
 }
+
+#endif
\ No newline at end of file
diff --git a/Engine/Plugins/Runtime/ComputeFramework/Source/ComputeFramework/Public/ComputeFramework/ComputeKernelSource.h b/Engine/Plugins/Runtime/ComputeFramework/Source/ComputeFramework/Public/ComputeFramework/ComputeKernelSource.h
--- a/Engine/Plugins/Runtime/ComputeFramework/Source/ComputeFramework/Public/ComputeFramework/ComputeKernelSource.h
+++ b/Engine/Plugins/Runtime/ComputeFramework/Source/ComputeFramework/Public/ComputeFramework/ComputeKernelSource.h
@@ -20,8 +20,10 @@ class UComputeKernelSource : public UObject
 	GENERATED_BODY()
 
 public:
+#if WITH_EDITOR
 	/** Get kernel source code ready for HLSL compilation. */
 	virtual FString GetSource() const PURE_VIRTUAL(UComputeKernelSource::GetSource, return {};);
+#endif
 
 	/** Kernel entry point. */
 	UPROPERTY(VisibleAnywhere, AssetRegistrySearchable, Category = "Kernel")
@@ -63,9 +65,11 @@ public:
 	UPROPERTY(EditAnywhere, Category = "Kernel")
 	FString SourceText;
 
+#if WITH_EDITOR
 	//~ Begin UComputeKernelSource Interface.
 	UE_API virtual FString GetSource() const override;
 	//~ End UComputeKernelSource Interface.
+#endif
 };
 
 #undef UE_API
