ComputeFramework: Fix localization name collisions.
#rnx
#Synced-CL 52465805

diff --git a/Engine/Plugins/Runtime/ComputeFramework/Source/EditableComputeGraphEditor/Private/ComputeGraphNavigator.cpp b/Engine/Plugins/Runtime/ComputeFramework/Source/EditableComputeGraphEditor/Private/ComputeGraphNavigator.cpp
--- a/Engine/Plugins/Runtime/ComputeFramework/Source/EditableComputeGraphEditor/Private/ComputeGraphNavigator.cpp
+++ b/Engine/Plugins/Runtime/ComputeFramework/Source/EditableComputeGraphEditor/Private/ComputeGraphNavigator.cpp
@@ -112,7 +112,7 @@ public:
 					SNew(SButton)
 					.ButtonStyle(FAppStyle::Get(), "SimpleButton")
 					.ContentPadding(FMargin(4.f, 0.f))
-					.ToolTipText(LOCTEXT("AddItemTooltip", "Add new item"))
+					.ToolTipText(LOCTEXT("Navigator_AddItemTooltip", "Add new item"))
 					.OnClicked_Lambda([this]()
 					{
 						OnAddItem.ExecuteIfBound(ItemPtr->Kind);
@@ -370,20 +370,20 @@ TSharedPtr<SWidget> SComputeGraphNavigator::BuildContextMenu()
 	FMenuBuilder MenuBuilder(/*bInShouldCloseWindowAfterMenuSelection=*/true, nullptr);
 
 	MenuBuilder.AddMenuEntry(
-		LOCTEXT("DeleteItem", "Delete"),
-		LOCTEXT("DeleteItemTooltip", "Remove this item from the graph description."),
+		LOCTEXT("Navigator_DeleteItem", "Delete"),
+		LOCTEXT("Navigator_DeleteItemTooltip", "Remove this item from the graph description."),
 		FSlateIcon(FAppStyle::GetAppStyleSetName(), "Icons.Delete"),
 		FUIAction(FExecuteAction::CreateSP(this, &SComputeGraphNavigator::OnDeleteSelected)));
 
 	MenuBuilder.AddMenuEntry(
-		LOCTEXT("DuplicateItem", "Duplicate"),
-		LOCTEXT("DuplicateItemTooltip", "Add a copy of this item with a new name."),
+		LOCTEXT("Navigator_DuplicateItem", "Duplicate"),
+		LOCTEXT("Navigator_DuplicateItemTooltip", "Add a copy of this item with a new name."),
 		FSlateIcon(FAppStyle::GetAppStyleSetName(), "Icons.Duplicate"),
 		FUIAction(FExecuteAction::CreateSP(this, &SComputeGraphNavigator::OnDuplicateSelected)));
 
 	MenuBuilder.AddMenuEntry(
-		LOCTEXT("RenameItem", "Rename"),
-		LOCTEXT("RenameItemTooltip", "Rename this item inline."),
+		LOCTEXT("Navigator_RenameItem", "Rename"),
+		LOCTEXT("Navigator_RenameItemTooltip", "Rename this item inline."),
 		FSlateIcon(FAppStyle::GetAppStyleSetName(), "GenericCommands.Rename"),
 		FUIAction(FExecuteAction::CreateSP(this, &SComputeGraphNavigator::OnRenameSelected)));
 
@@ -455,7 +455,7 @@ bool SComputeGraphNavigator::VerifyRenameText(FText const& NewText, FText& OutEr
 	const FString Trimmed = NewText.ToString().TrimStartAndEnd();
 	if (Trimmed.IsEmpty())
 	{
-		OutError = LOCTEXT("RenameEmpty", "Name cannot be empty.");
+		OutError = LOCTEXT("Navigator_RenameEmpty", "Name cannot be empty.");
 		return false;
 	}
 
@@ -467,7 +467,7 @@ bool SComputeGraphNavigator::VerifyRenameText(FText const& NewText, FText& OutEr
 		{
 			if (Child->Kind == Kind && Child->Index != Index && Child->Name == Candidate)
 			{
-				OutError = FText::Format(LOCTEXT("RenameDuplicate", "'{0}' is already in use."), NewText);
+				OutError = FText::Format(LOCTEXT("Navigator_RenameDuplicate", "'{0}' is already in use."), NewText);
 				return false;
 			}
 		}
diff --git a/Engine/Plugins/Runtime/ComputeFramework/Source/EditableComputeGraphEditor/Private/EditableComputeGraphEditorToolkit.cpp b/Engine/Plugins/Runtime/ComputeFramework/Source/EditableComputeGraphEditor/Private/EditableComputeGraphEditorToolkit.cpp
--- a/Engine/Plugins/Runtime/ComputeFramework/Source/EditableComputeGraphEditor/Private/EditableComputeGraphEditorToolkit.cpp
+++ b/Engine/Plugins/Runtime/ComputeFramework/Source/EditableComputeGraphEditor/Private/EditableComputeGraphEditorToolkit.cpp
@@ -737,7 +737,7 @@ void FEditableComputeGraphEditorToolkit::OnDeleteItem(EComputeGraphItemKind Kind
 		return;
 	}
 
-	FScopedTransaction Transaction(LOCTEXT("DeleteItem", "Delete Compute Graph Item"));
+	FScopedTransaction Transaction(LOCTEXT("Toolkit_DeleteItem", "Delete Compute Graph Item"));
 	Asset->Modify();
 
 	FComputeGraphDesc& Desc = Asset->GetGraphDescription();
@@ -782,7 +782,7 @@ void FEditableComputeGraphEditorToolkit::OnDuplicateItem(EComputeGraphItemKind K
 		return;
 	}
 
-	FScopedTransaction Transaction(LOCTEXT("DuplicateItem", "Duplicate Compute Graph Item"));
+	FScopedTransaction Transaction(LOCTEXT("Toolkit_DuplicateItem", "Duplicate Compute Graph Item"));
 	Asset->Modify();
 
 	FComputeGraphDesc& Desc = Asset->GetGraphDescription();
@@ -853,7 +853,7 @@ void FEditableComputeGraphEditorToolkit::OnRenameItem(EComputeGraphItemKind Kind
 		return;
 	}
 
-	FScopedTransaction Transaction(LOCTEXT("RenameItem", "Rename Compute Graph Item"));
+	FScopedTransaction Transaction(LOCTEXT("Toolkit_RenameItem", "Rename Compute Graph Item"));
 	Asset->Modify();
 
 	FComputeGraphDesc& Desc = Asset->GetGraphDescription();
