[ComputeFramework] Support cooked editor workflows and reduce cooked runtime memory usage.

Mark UComputeKernelSource shader source properties as editor-only so cooked game
builds don't pay memory for data only needed at compile time.

In cooked editor builds the editor-only data is still available, so we can still
compile shaders when needed. Add CacheCookedResourceShadersForRendering() to first
attempt to use inline shader maps from the cooked asset before falling back to
compilation. Add FComputeKernelResource::CacheInlineShaders() to handle the render
thread promotion for loaded inline shader maps.

The ContentWorker is a special case: it cooks assets but intentionally omits
shaders, so the inline shader map fallback is skipped and shaders are compiled at
runtime. This requires the cooked asset to include optional editor-only data.
Register UOptimusDeformer with AutomaticOptionalInclusionAssetType (via Editor.ini)
to ensure this data is available when loading ContentWorker packages in the editor.

#Synced-CL 52466221

diff --git a/Engine/Plugins/Animation/DeformerGraph/Config/Editor.ini b/Engine/Plugins/Animation/DeformerGraph/Config/Editor.ini
new file mode 100644
--- /dev/null
+++ b/Engine/Plugins/Animation/DeformerGraph/Config/Editor.ini
@@ -0,0 +1,2 @@
+[CookSettings]
++AutomaticOptionalInclusionAssetType=/Script/OptimusCore.OptimusDeformer
\ No newline at end of file
diff --git a/Engine/Plugins/Animation/DeformerGraph/Source/OptimusCore/Private/Nodes/OptimusNode_ComputeKernelBase.cpp b/Engine/Plugins/Animation/DeformerGraph/Source/OptimusCore/Private/Nodes/OptimusNode_ComputeKernelBase.cpp
--- a/Engine/Plugins/Animation/DeformerGraph/Source/OptimusCore/Private/Nodes/OptimusNode_ComputeKernelBase.cpp
+++ b/Engine/Plugins/Animation/DeformerGraph/Source/OptimusCore/Private/Nodes/OptimusNode_ComputeKernelBase.cpp
@@ -54,6 +54,7 @@ static void CopyValueType(FShaderValueTypeHandle InValueType,  FShaderParamTypeD
 	OutParamDef.ResetTypeDeclaration();
 }
 
+#if WITH_EDITOR
 FOptimus_ComputeKernelResult UOptimusNode_ComputeKernelBase::CreateComputeKernel(
 	UObject *InKernelSourceOuter,
 	const FOptimusPinTraversalContext& InTraversalContext,
@@ -262,7 +263,7 @@ FOptimus_ComputeKernelResult UOptimusNode_ComputeKernelBase::CreateComputeKernel
 	
 	return FOptimus_ComputeKernelResult(TInPlaceType<UOptimusKernelSource*>(), KernelSource);
 }
-
+#endif // WITH_EDITOR
 
 TOptional<FText> UOptimusNode_ComputeKernelBase::ValidateForCompile(const FOptimusPinTraversalContext& InContext) const
 {
@@ -391,7 +392,7 @@ const TCHAR* UOptimusNode_ComputeKernelBase::GetNumThreadsReservedName()
 	return NumThreadsReservedName;
 }
 
-
+#if WITH_EDITOR
 TOptional<FText> UOptimusNode_ComputeKernelBase::ProcessInputPinForComputeKernel(
 	bool bSupportUnifiedDispatchOnly,
 	const FOptimusPinTraversalContext& InTraversalContext,
@@ -816,5 +817,6 @@ TOptional<FText> UOptimusNode_ComputeKernelBase::BindKernelDataInterfaceForCompu
 
 	return {};
 }
+#endif // WITH_EDITOR
 
 #undef LOCTEXT_NAMESPACE
diff --git a/Engine/Plugins/Animation/DeformerGraph/Source/OptimusCore/Private/Nodes/OptimusNode_ComputeKernelBase.h b/Engine/Plugins/Animation/DeformerGraph/Source/OptimusCore/Private/Nodes/OptimusNode_ComputeKernelBase.h
--- a/Engine/Plugins/Animation/DeformerGraph/Source/OptimusCore/Private/Nodes/OptimusNode_ComputeKernelBase.h
+++ b/Engine/Plugins/Animation/DeformerGraph/Source/OptimusCore/Private/Nodes/OptimusNode_ComputeKernelBase.h
@@ -35,6 +35,7 @@ public:
 	/** Implement this to return additional source code that the kernel requires. */
 	virtual TArray<TObjectPtr<UComputeSource>> GetAdditionalSources() const { return {}; }
 	
+#if WITH_EDITOR
 	// IOptimusComputeKernelProvider
 	FOptimus_ComputeKernelResult CreateComputeKernel(
 		UObject* InKernelSourceOuter,
@@ -45,6 +46,7 @@ public:
 		FOptimus_InterfaceBindingMap& OutInputDataBindings,
 		FOptimus_InterfaceBindingMap& OutOutputDataBindings
 	) const override;
+#endif
 
 	FOptimusExecutionDomain GetExecutionDomain() const override PURE_VIRTUAL(UOptimusNode_ComputeKernelBase::GetExecutionDomain, return {}; );
 	const UOptimusNodePin* GetPrimaryGroupPin() const override PURE_VIRTUAL(UOptimusNode_ComputeKernelBase::GetPrimaryGroupPin, return {}; );  
@@ -96,6 +98,7 @@ protected:
 	static FString GetReadFunctionName(const FString& InBindingName);
 	static const TCHAR* GetNumThreadsReservedName();
 	
+#if WITH_EDITOR
 private:
 	TOptional<FText> BindKernelDataInterfaceForComputeKernel(
 		const UOptimusComponentSourceBinding* InKernelPrimaryComponentSourceBinding,
@@ -129,6 +132,6 @@ private:
 		FOptimus_InterfaceBindingMap& OutOutputDataBindings,
 		bool& bOutIsUnifiedDispatch
 		) const;
-
+#endif // WITH_EDITOR
 };
 
diff --git a/Engine/Plugins/Animation/DeformerGraph/Source/OptimusCore/Private/OptimusDeformer.cpp b/Engine/Plugins/Animation/DeformerGraph/Source/OptimusCore/Private/OptimusDeformer.cpp
--- a/Engine/Plugins/Animation/DeformerGraph/Source/OptimusCore/Private/OptimusDeformer.cpp
+++ b/Engine/Plugins/Animation/DeformerGraph/Source/OptimusCore/Private/OptimusDeformer.cpp
@@ -1494,7 +1494,7 @@ static void CollectNodes(
 	}	
 }
 
-
+#if WITH_EDITOR
 bool UOptimusDeformer::Compile()
 {
 	if (!GetUpdateGraph())
@@ -1605,6 +1605,7 @@ bool UOptimusDeformer::Compile()
 	
 	return true;
 }
+#endif // WITH_EDITOR
 
 TArray<UOptimusNode*> UOptimusDeformer::GetAllNodesOfClass(UClass* InNodeClass) const
 {
@@ -1718,6 +1719,7 @@ struct FOptimusInstancedPin
 	}
 };
 
+#if WITH_EDITOR
 FOptimusNodeGraphCompilationResult UOptimusDeformer::CompileNodeGraphToComputeGraphs(
 	const UOptimusNodeGraph* InNodeGraph,
 	TFunction<void(EOptimusDiagnosticLevel, FText, const UObject*)> InErrorReporter
@@ -3208,6 +3210,7 @@ FOptimusNodeGraphCompilationResult UOptimusDeformer::CompileNodeGraphToComputeGr
 
 	return Result;
 }
+#endif // WITH_EDITOR
 
 void UOptimusDeformer::OnDataTypeChanged(FName InTypeName)
 {
@@ -3504,6 +3507,7 @@ void UOptimusDeformer::PostLoad()
 		}
 	}
 
+#if WITH_EDITOR
 	FOptimusObjectVersion::Type ObjectVersion = static_cast<FOptimusObjectVersion::Type>(GetLinkerCustomVersion(FOptimusObjectVersion::GUID));
 	if (ObjectVersion < FOptimusObjectVersion::ReparentResourcesAndVariables)
 	{
@@ -3620,6 +3624,7 @@ void UOptimusDeformer::PostLoad()
 			Compile();
 		}
 	}
+#endif // WITH_EDITOR
 }
 
 
diff --git a/Engine/Plugins/Animation/DeformerGraph/Source/OptimusCore/Public/IOptimusComputeKernelProvider.h b/Engine/Plugins/Animation/DeformerGraph/Source/OptimusCore/Public/IOptimusComputeKernelProvider.h
--- a/Engine/Plugins/Animation/DeformerGraph/Source/OptimusCore/Public/IOptimusComputeKernelProvider.h
+++ b/Engine/Plugins/Animation/DeformerGraph/Source/OptimusCore/Public/IOptimusComputeKernelProvider.h
@@ -71,6 +71,7 @@ class IOptimusComputeKernelProvider
 	GENERATED_BODY()
 
 public:
+#if WITH_EDITOR
 	/**
 	 * Return an UOptimusKernelSource object, from a compute kernel node state that implements
 	 * this interface.
@@ -89,6 +90,7 @@ public:
 		FOptimus_InterfaceBindingMap& OutInputDataBindings,
 		FOptimus_InterfaceBindingMap& OutOutputDataBindings
 	) const = 0;
+#endif
 
 	/** Returns the execution domain that this kernel should iterate over */
 	virtual FOptimusExecutionDomain GetExecutionDomain() const = 0;
diff --git a/Engine/Plugins/Animation/DeformerGraph/Source/OptimusCore/Public/OptimusDeformer.h b/Engine/Plugins/Animation/DeformerGraph/Source/OptimusCore/Public/OptimusDeformer.h
--- a/Engine/Plugins/Animation/DeformerGraph/Source/OptimusCore/Public/OptimusDeformer.h
+++ b/Engine/Plugins/Animation/DeformerGraph/Source/OptimusCore/Public/OptimusDeformer.h
@@ -300,7 +300,9 @@ public:
 
 	
 	/// Graph compilation
+#if WITH_EDITOR
 	UE_API bool Compile();
+#endif
 
 	/** Returns a multicast delegate that can be subscribed to listen for the start of compilation. */
 	FOptimusCompileBegin& GetCompileBeginDelegate()  { return CompileBeginDelegate; }
@@ -542,11 +544,13 @@ private:
 	UE_API void OnGraphRenamedOrRemoved(UOptimusNodeGraph* InGraph) const;
 	UE_API void UpdateFunctionReferenceNodeDisplayName(const FOptimusFunctionGraphIdentifier& InRenamedFunction);
 	
+#if WITH_EDITOR
 	/// Compile a node graph to a compute graph. Returns one or two complete compute graphs if compilation succeeded. 
 	UE_API FOptimusNodeGraphCompilationResult CompileNodeGraphToComputeGraphs(
 		const UOptimusNodeGraph *InNodeGraph,
 		TFunction<void(EOptimusDiagnosticLevel, FText, const UObject*)> InErrorReporter
 		);
+#endif
 
 	UE_API void OnDataTypeChanged(FName InTypeName);
 
diff --git a/Engine/Plugins/Animation/DeformerGraph/Source/OptimusCore/Public/OptimusKernelSource.h b/Engine/Plugins/Animation/DeformerGraph/Source/OptimusCore/Public/OptimusKernelSource.h
--- a/Engine/Plugins/Animation/DeformerGraph/Source/OptimusCore/Public/OptimusKernelSource.h
+++ b/Engine/Plugins/Animation/DeformerGraph/Source/OptimusCore/Public/OptimusKernelSource.h
@@ -12,6 +12,7 @@ class UOptimusKernelSource :
 	GENERATED_BODY()
 
 public:
+#if WITH_EDITOR
 	void SetSource(FString const& InSource)
 	{
 		Source = InSource;
@@ -27,6 +28,8 @@ public:
 #endif
 
 protected:
+#if WITH_EDITORONLY_DATA
 	UPROPERTY()
 	FString Source;
+#endif
 };
diff --git a/Engine/Plugins/PCG/Source/PCG/Private/Compute/PCGComputeKernelSource.h b/Engine/Plugins/PCG/Source/PCG/Private/Compute/PCGComputeKernelSource.h
--- a/Engine/Plugins/PCG/Source/PCG/Private/Compute/PCGComputeKernelSource.h
+++ b/Engine/Plugins/PCG/Source/PCG/Private/Compute/PCGComputeKernelSource.h
@@ -24,8 +24,11 @@ public:
 	{
 		Source = InSource;
 	}
+#endif
 
 protected:
+#if WITH_EDITOR
+	// Transient string used only during compilation.
 	FString Source;
 #endif
 };
diff --git a/Engine/Plugins/Runtime/ComputeFramework/Source/ComputeFramework/Private/ComputeFramework/ComputeKernelShared.h b/Engine/Plugins/Runtime/ComputeFramework/Source/ComputeFramework/Private/ComputeFramework/ComputeKernelShared.h
--- a/Engine/Plugins/Runtime/ComputeFramework/Source/ComputeFramework/Private/ComputeFramework/ComputeKernelShared.h
+++ b/Engine/Plugins/Runtime/ComputeFramework/Source/ComputeFramework/Private/ComputeFramework/ComputeKernelShared.h
@@ -308,6 +308,7 @@ public:
 	 */
 	bool CacheShaders(EShaderPlatform InPlatform, const ITargetPlatform* TargetPlatform, bool bApplyCompletedShaderMapForRendering, bool bSynchronous);
 	bool CacheShaders(const FComputeKernelShaderMapId& InShaderMapId, EShaderPlatform InPlatform, bool bApplyCompletedShaderMapForRendering, bool bSynchronous);
+	bool CacheInlineShaders(EShaderPlatform InPlatform);
 
 	/**
 	 * Should the shader for this kernel with the given platform and shader type combination be compiled
diff --git a/Engine/Plugins/Runtime/ComputeFramework/Source/ComputeFramework/Private/ComputeGraph.cpp b/Engine/Plugins/Runtime/ComputeFramework/Source/ComputeFramework/Private/ComputeGraph.cpp
--- a/Engine/Plugins/Runtime/ComputeFramework/Source/ComputeFramework/Private/ComputeGraph.cpp
+++ b/Engine/Plugins/Runtime/ComputeFramework/Source/ComputeFramework/Private/ComputeGraph.cpp
@@ -25,7 +25,7 @@
 #include "UObject/ObjectSaveContext.h"
 #include "UObject/Package.h"
 
-#if WITH_EDITOR
+#if WITH_EDITORONLY_DATA
 #include "Cooker/CookDependencyContext.h"
 #include "Cooker/CookEvents.h"
 #endif
@@ -93,7 +93,7 @@ void UComputeGraph::PostLoad()
 		}
 	}
 	
-#if WITH_EDITOR
+#if WITH_EDITORONLY_DATA
 	// PostLoad our kernel dependencies before any compiling.
 	for (UComputeKernel* Kernel : KernelInvocations)
 	{
@@ -117,7 +117,7 @@ void UComputeGraph::PostLoad()
 	}
 }
 
-#if WITH_EDITOR
+#if WITH_EDITORONLY_DATA
 
 /** Shader dependency serialization code used for incremental cooking. */
 namespace ComputeFramework::Private
@@ -165,10 +165,10 @@ namespace ComputeFramework::Private
 
 UE_COOK_DEPENDENCY_FUNCTION(HashComputeGraphDependenciesForCook, ComputeFramework::Private::HashDependenciesForCook);
 
-#endif // WITH_EDITOR
+#endif // WITH_EDITORONLY_DATA
 
 
-#if WITH_EDITOR
+#if WITH_EDITORONLY_DATA
 void UComputeGraph::OnCookEvent(UE::Cook::ECookEvent CookEvent, UE::Cook::FCookEventContext& CookContext)
 {
 	Super::OnCookEvent(CookEvent, CookContext);
@@ -513,7 +513,7 @@ enum class EComputeKernelCompilationFlags
 
 void UComputeGraph::UpdateResources(bool bSync)
 {
-#if WITH_EDITOR
+#if WITH_EDITORONLY_DATA
 	uint32 CompilationFlags = (uint32)EComputeKernelCompilationFlags::ApplyCompletedShaderMapForRendering;
 
 	CompilationFlags |= bSync ? (uint32)EComputeKernelCompilationFlags::Synchronous : (uint32)EComputeKernelCompilationFlags::None;
@@ -756,7 +756,7 @@ void UComputeGraph::GetDataInterfaceInputOutputMasks(int32 InDataInterfaceIndex,
 	}
 }
 
-#if WITH_EDITOR
+#if WITH_EDITORONLY_DATA
 
 namespace
 {
@@ -1011,63 +1011,97 @@ FString UComputeGraph::BuildKernelSource(
 	return HLSL;
 }
 
-void UComputeGraph::CacheResourceShadersForRendering(uint32 CompilationFlags)
+bool UComputeGraph::CacheCookedResourceShadersForRendering()
 {
-	if (FApp::CanEverRender())
+	for (int32 KernelIndex = 0; KernelIndex < KernelInvocations.Num(); ++KernelIndex)
 	{
-		KernelResources.SetNum(KernelInvocations.Num());
-
-		KernelResourceIndicesPendingShaderCompilation.Reset();
-		for (int32 KernelIndex = 0; KernelIndex < KernelInvocations.Num(); ++KernelIndex)
+		UComputeKernel const* Kernel = KernelInvocations[KernelIndex];
+		if (Kernel == nullptr || Kernel->KernelSource == nullptr)
 		{
-			KernelResourceIndicesPendingShaderCompilation.Add(KernelIndex);
+			// Skip kernels that would not have been cooked.
+			continue;
 		}
-		
-		for (int32 KernelIndex = 0; KernelIndex < KernelInvocations.Num(); ++KernelIndex)
+		if (KernelResources.IsValidIndex(KernelIndex) && KernelResources[KernelIndex].CacheInlineShaders())
 		{
-			UComputeKernel* Kernel = KernelInvocations[KernelIndex];
-			if (Kernel == nullptr || Kernel->KernelSource == nullptr)
-			{
-				KernelResources[KernelIndex].Reset();
-				continue;
-			}
+			// Found and cached valid inline shader map.
+			continue;
+		}
+		// Some required resource or inline shaders not found.
+		return false;
+	}
+	return true;
+}
 
-			TMap<FString, FString> AdditionalSources = GatherAdditionalSources(Kernel->KernelSource->AdditionalSources);
+void UComputeGraph::CacheResourceShadersForRendering(uint32 CompilationFlags)
+{
+	// Don't build shaders in non-rendering apps such as commandlets.
+	if (!FApp::CanEverRender())
+	{
+		return;
+	}
 
-			FString ShaderHashKey;
-			TMap<FString, FString> GeneratedSources;
-			TSharedPtr<FComputeKernelDefinitionSet> ShaderDefinitionSet = MakeShared<FComputeKernelDefinitionSet>();
-			TSharedPtr<FComputeKernelPermutationVector> ShaderPermutationVector = MakeShared<FComputeKernelPermutationVector>();
-			TUniquePtr<FShaderParametersMetadataAllocations> ShaderParameterMetadataAllocations = MakeUnique<FShaderParametersMetadataAllocations>();
-
-			FString ShaderEntryPoint = Kernel->KernelSource->EntryPoint;
-			FString ShaderFriendlyName = GetOuter()->GetName() / GetFName().GetPlainNameString() / ShaderEntryPoint;
-			FString ShaderSource = BuildKernelSource(KernelIndex, *Kernel->KernelSource, AdditionalSources, ShaderHashKey, GeneratedSources, *ShaderDefinitionSet, *ShaderPermutationVector);
-			FShaderParametersMetadata* ShaderParameterMetadata = BuildKernelShaderMetadata(KernelIndex, *ShaderParameterMetadataAllocations);
-
-			const ERHIFeatureLevel::Type CacheFeatureLevel = GMaxRHIFeatureLevel;
-			const EShaderPlatform ShaderPlatform = GShaderPlatformForFeatureLevel[CacheFeatureLevel];
-			FComputeKernelResource* KernelResource = KernelResources[KernelIndex].GetOrCreate();
-
-			// Now we have all the information that the KernelResource will need for compilation.
-			KernelResource->SetupResource(
-				CacheFeatureLevel, 
-				ShaderFriendlyName,
-				ShaderEntryPoint, 
-				ShaderHashKey,
-				ShaderSource,
-				AdditionalSources,
-				GeneratedSources,
-				ShaderDefinitionSet,
-				ShaderPermutationVector,
-				ShaderParameterMetadataAllocations,
-				ShaderParameterMetadata,
-				GetOutermost()->GetFName());
-
-			KernelResource->OnCompilationComplete().BindUObject(this, &UComputeGraph::ShaderCompileCompletionCallback);
-
-			CacheShadersForResource(ShaderPlatform, nullptr, CompilationFlags, KernelResource);
+	// Don't build shaders if we are working with cooked data that should contain built inline shader maps.
+	// Note that the ContentWorker is a special case that can have cooked resources which don't contain shaders.
+	// For that case CacheCookedResourceShadersForRendering() returns false and we will fall back to the compilation path below.
+	// Note that this fall back path depends on cooking with optional editor only data so that we can compile shaders using the cooked asset.
+	// UOptimusDeformer registers with AutomaticOptionalInclusionAssetType to make that happen. Other usages would need to do the same.
+	if (GetPackage()->HasAnyPackageFlags(PKG_Cooked) && CacheCookedResourceShadersForRendering())
+	{
+		return;
+	}
+
+	KernelResources.SetNum(KernelInvocations.Num());
+
+	KernelResourceIndicesPendingShaderCompilation.Reset();
+	for (int32 KernelIndex = 0; KernelIndex < KernelInvocations.Num(); ++KernelIndex)
+	{
+		KernelResourceIndicesPendingShaderCompilation.Add(KernelIndex);
+	}
+		
+	for (int32 KernelIndex = 0; KernelIndex < KernelInvocations.Num(); ++KernelIndex)
+	{
+		UComputeKernel* Kernel = KernelInvocations[KernelIndex];
+		if (Kernel == nullptr || Kernel->KernelSource == nullptr)
+		{
+			KernelResources[KernelIndex].Reset();
+			continue;
 		}
+
+		TMap<FString, FString> AdditionalSources = GatherAdditionalSources(Kernel->KernelSource->AdditionalSources);
+
+		FString ShaderHashKey;
+		TMap<FString, FString> GeneratedSources;
+		TSharedPtr<FComputeKernelDefinitionSet> ShaderDefinitionSet = MakeShared<FComputeKernelDefinitionSet>();
+		TSharedPtr<FComputeKernelPermutationVector> ShaderPermutationVector = MakeShared<FComputeKernelPermutationVector>();
+		TUniquePtr<FShaderParametersMetadataAllocations> ShaderParameterMetadataAllocations = MakeUnique<FShaderParametersMetadataAllocations>();
+
+		FString ShaderEntryPoint = Kernel->KernelSource->EntryPoint;
+		FString ShaderFriendlyName = GetOuter()->GetName() / GetFName().GetPlainNameString() / ShaderEntryPoint;
+		FString ShaderSource = BuildKernelSource(KernelIndex, *Kernel->KernelSource, AdditionalSources, ShaderHashKey, GeneratedSources, *ShaderDefinitionSet, *ShaderPermutationVector);
+		FShaderParametersMetadata* ShaderParameterMetadata = BuildKernelShaderMetadata(KernelIndex, *ShaderParameterMetadataAllocations);
+
+		const ERHIFeatureLevel::Type CacheFeatureLevel = GMaxRHIFeatureLevel;
+		const EShaderPlatform ShaderPlatform = GShaderPlatformForFeatureLevel[CacheFeatureLevel];
+		FComputeKernelResource* KernelResource = KernelResources[KernelIndex].GetOrCreate();
+
+		// Now we have all the information that the KernelResource will need for compilation.
+		KernelResource->SetupResource(
+			CacheFeatureLevel, 
+			ShaderFriendlyName,
+			ShaderEntryPoint, 
+			ShaderHashKey,
+			ShaderSource,
+			AdditionalSources,
+			GeneratedSources,
+			ShaderDefinitionSet,
+			ShaderPermutationVector,
+			ShaderParameterMetadataAllocations,
+			ShaderParameterMetadata,
+			GetOutermost()->GetFName());
+
+		KernelResource->OnCompilationComplete().BindUObject(this, &UComputeGraph::ShaderCompileCompletionCallback);
+
+		CacheShadersForResource(ShaderPlatform, nullptr, CompilationFlags, KernelResource);
 	}
 }
 
@@ -1238,7 +1272,7 @@ void UComputeGraph::ClearAllCachedCookedPlatformData()
 	KernelResources.Reset();
 }
 
-#endif // WITH_EDITOR
+#endif // WITH_EDITORONLY_DATA
 
 void UComputeGraph::FComputeKernelResourceSet::Reset()
 {
@@ -1335,10 +1369,9 @@ void UComputeGraph::FComputeKernelResourceSet::Serialize(FArchive& Ar)
 		else
 #endif // WITH_EDITORONLY_DATA
 		{
-			int32 NumResources = 0;
-			Ar << NumResources;
-
-			for (int32 ResourceIndex = 0; ResourceIndex < NumResources; ++ResourceIndex)
+			int32 NumLoadedResources = 0;
+			Ar << NumLoadedResources;
+			for (int32 ResourceIndex = 0; ResourceIndex < NumLoadedResources; ++ResourceIndex)
 			{
 				TUniquePtr<FComputeKernelResource> Resource = MakeUnique<FComputeKernelResource>();
 				Resource->SerializeShaderMap(Ar);
@@ -1386,6 +1419,19 @@ void UComputeGraph::FComputeKernelResourceSet::ProcessSerializedShaderMaps()
 #endif
 }
 
+bool UComputeGraph::FComputeKernelResourceSet::CacheInlineShaders()
+{
+#if WITH_EDITORONLY_DATA
+	const ERHIFeatureLevel::Type CacheFeatureLevel = GMaxRHIFeatureLevel;
+	if (FComputeKernelResource* KernelResource = KernelResourcesByFeatureLevel[CacheFeatureLevel].Get())
+	{
+		const EShaderPlatform ShaderPlatform = GShaderPlatformForFeatureLevel[CacheFeatureLevel];
+		return KernelResource->CacheInlineShaders(ShaderPlatform);
+	}
+#endif
+	return false;
+}
+
 #if WITH_EDITORONLY_DATA
 void UComputeGraph::FComputeKernelResourceSet::CancelCompilation()
 {
diff --git a/Engine/Plugins/Runtime/ComputeFramework/Source/ComputeFramework/Private/ComputeKernelShared.cpp b/Engine/Plugins/Runtime/ComputeFramework/Source/ComputeFramework/Private/ComputeKernelShared.cpp
--- a/Engine/Plugins/Runtime/ComputeFramework/Source/ComputeFramework/Private/ComputeKernelShared.cpp
+++ b/Engine/Plugins/Runtime/ComputeFramework/Source/ComputeFramework/Private/ComputeKernelShared.cpp
@@ -415,6 +415,11 @@ bool FComputeKernelResource::CacheShaders(const FComputeKernelShaderMapId& InSha
 	return bSucceeded;
 }
 
+bool FComputeKernelResource::CacheInlineShaders(EShaderPlatform InPlatform)
+{
+	return bContainsInlineShaders && CacheShaders(InPlatform, nullptr, true, true);
+}
+
 #if WITH_EDITOR
 void FComputeKernelResource::FinishCompilation()
 {
diff --git a/Engine/Plugins/Runtime/ComputeFramework/Source/ComputeFramework/Public/ComputeFramework/ComputeGraph.h b/Engine/Plugins/Runtime/ComputeFramework/Source/ComputeFramework/Public/ComputeFramework/ComputeGraph.h
--- a/Engine/Plugins/Runtime/ComputeFramework/Source/ComputeFramework/Public/ComputeFramework/ComputeGraph.h
+++ b/Engine/Plugins/Runtime/ComputeFramework/Source/ComputeFramework/Public/ComputeFramework/ComputeGraph.h
@@ -157,19 +157,19 @@ protected:
 
 private:
 	/** Build the shader metadata which describes bindings for a kernel with its linked data interfaces.*/
-	COMPUTEFRAMEWORK_API FShaderParametersMetadata* BuildKernelShaderMetadata(int32 InKernelIndex, FShaderParametersMetadataAllocations& InOutAllocations) const;
+	FShaderParametersMetadata* BuildKernelShaderMetadata(int32 InKernelIndex, FShaderParametersMetadataAllocations& InOutAllocations) const;
 	/** Build the shader permutation vectors for all kernels in the graph. */
-	COMPUTEFRAMEWORK_API void BuildShaderPermutationVectors(TArray<FComputeKernelPermutationVector>& OutShaderPermutationVectors) const;
+	void BuildShaderPermutationVectors(TArray<FComputeKernelPermutationVector>& OutShaderPermutationVectors) const;
 	/** Create the render proxy. */
-	COMPUTEFRAMEWORK_API FComputeGraphRenderProxy* CreateRenderProxy() const;
+	FComputeGraphRenderProxy* CreateRenderProxy() const;
 	/** Release the render proxy. */
-	COMPUTEFRAMEWORK_API void ReleaseRenderProxy(FComputeGraphRenderProxy* InProxy) const;
+	void ReleaseRenderProxy(FComputeGraphRenderProxy* InProxy) const;
 	/** Returns input and output masks for the given data interface. */
-	COMPUTEFRAMEWORK_API void GetDataInterfaceInputOutputMasks(int32 InDataInterfaceIndex, uint64& OutInputMask, uint64& OutOutputMask) const;
+	void GetDataInterfaceInputOutputMasks(int32 InDataInterfaceIndex, uint64& OutInputMask, uint64& OutOutputMask) const;
 
 #if WITH_EDITOR
 	/** Build the HLSL source for a kernel with its linked data interfaces. */
-	COMPUTEFRAMEWORK_API FString BuildKernelSource(
+	FString BuildKernelSource(
 		int32 KernelIndex, 
 		UComputeKernelSource const& InKernelSource,
 		TMap<FString, FString> const& InAdditionalSources,
@@ -179,17 +179,20 @@ private:
 		FComputeKernelPermutationVector& OutPermutationVector) const;
 
 	/** Cache shader resources for all kernels in the graph. */
-	COMPUTEFRAMEWORK_API void CacheResourceShadersForRendering(uint32 CompilationFlags);
+	void CacheResourceShadersForRendering(uint32 CompilationFlags);
+	
+	/** For cooked packages, promote inline shader maps to the render thread. Returns false if any kernel was missing an inline shader map. */
+	bool CacheCookedResourceShadersForRendering();
 
 	/** Cache shader resources for a specific compute kernel. This will trigger any required shader compilation. */
-	static COMPUTEFRAMEWORK_API void CacheShadersForResource(
+	static void CacheShadersForResource(
 		EShaderPlatform ShaderPlatform,
 		ITargetPlatform const* TargetPlatform,
 		uint32 CompilationFlags,
 		FComputeKernelResource* Kernel);
 
 	/** Callback to handle result of kernel shader compilations. */
-	COMPUTEFRAMEWORK_API void ShaderCompileCompletionCallback(FComputeKernelResource const* KernelResource);
+	void ShaderCompileCompletionCallback(FComputeKernelResource const* KernelResource);
 #endif
 
 private:
@@ -225,6 +228,8 @@ private:
 		void Serialize(FArchive& Ar);
 		/** Apply shader maps found in Serialize(). Call this from PostLoad(). */
 		void ProcessSerializedShaderMaps();
+		/** Cache inline shader map for rendering. Returns false if no inline shader map was available. */
+		bool CacheInlineShaders();
 
 #if WITH_EDITORONLY_DATA
 		void CancelCompilation();
diff --git a/Engine/Plugins/Runtime/ComputeFramework/Source/ComputeFramework/Public/ComputeFramework/ComputeKernelSource.h b/Engine/Plugins/Runtime/ComputeFramework/Source/ComputeFramework/Public/ComputeFramework/ComputeKernelSource.h
--- a/Engine/Plugins/Runtime/ComputeFramework/Source/ComputeFramework/Public/ComputeFramework/ComputeKernelSource.h
+++ b/Engine/Plugins/Runtime/ComputeFramework/Source/ComputeFramework/Public/ComputeFramework/ComputeKernelSource.h
@@ -33,6 +33,7 @@ public:
 	UPROPERTY(VisibleAnywhere, AssetRegistrySearchable, Category = "Kernel")
 	FIntVector GroupSize = FIntVector(64, 1, 1);
 
+#if WITH_EDITORONLY_DATA
 	/** Base permutations exposed by the kernel. These will be extended by further permutations declared in any linked data providers. */
 	UPROPERTY(VisibleAnywhere, meta = (ShowOnlyInnerProperties), Category = "Kernel")
 	FComputeKernelPermutationSet PermutationSet;
@@ -52,6 +53,7 @@ public:
 	/* Named external outputs for the kernel. These must be fulfilled by linked data providers. */
 	UPROPERTY(VisibleAnywhere, EditFixedSize, Category = "External")
 	TArray<FShaderFunctionDefinition> ExternalOutputs;
+#endif // WITH_EDITORONLY_DATA
 };
 
 /** Simple concrete implementation of UComputeKernelSource with contained source text property. */
@@ -61,9 +63,11 @@ class UComputeKernelSourceWithText final : public UComputeKernelSource
 	GENERATED_BODY()
 
 public:
+#if WITH_EDITORONLY_DATA
 	/* Kernel source text. */
 	UPROPERTY(EditAnywhere, Category = "Kernel")
 	FString SourceText;
+#endif
 
 #if WITH_EDITOR
 	//~ Begin UComputeKernelSource Interface.
